Integrating //UE5/Dev-Cooker @ 16678003 to //UE5/Main (Zousar.Shaker-YEG-0943-Quaternary)

diff --git a/Engine/Source/Runtime/CoreUObject/Public/UObject/Package.h b/Engine/Source/Runtime/CoreUObject/Public/UObject/Package.h
--- a/Engine/Source/Runtime/CoreUObject/Public/UObject/Package.h
+++ b/Engine/Source/Runtime/CoreUObject/Public/UObject/Package.h
@@ -6,7 +6,8 @@
 #include "UObject/ObjectMacros.h"
 #include "UObject/UObjectGlobals.h"
 #include "UObject/Object.h"
-#include "UObject/PackageId.h"
+#include "IO/PackageId.h"
+#include "UObject/LinkerSave.h"
 #include "Misc/Guid.h"
 #include "Misc/OutputDeviceError.h"
 #include "Misc/ObjectThumbnail.h"
diff --git a/Engine/Source/Runtime/CoreUObject/Public/UObject/SavePackage.h b/Engine/Source/Runtime/CoreUObject/Public/UObject/SavePackage.h
--- a/Engine/Source/Runtime/CoreUObject/Public/UObject/SavePackage.h
+++ b/Engine/Source/Runtime/CoreUObject/Public/UObject/SavePackage.h
@@ -21,6 +21,7 @@ class FPackagePath;
 class FSavePackageContext;
 class FArchiveDiffMap;
 class FOutputDevice;
+class IPackageStoreWriter;
 
 /**
  * Struct to encapsulate arguments specific to saving one package
@@ -50,52 +51,21 @@ struct FSavePackageArgs
 	FSavePackageContext* SavePackageContext = nullptr;
 };
 
-class IPackageStoreWriter
-{
-public:
-	COREUOBJECT_API virtual ~IPackageStoreWriter();
-
-	struct FPackageInfo
-	{
-		FName	PackageName;
-		FString	LooseFilePath;
-		uint64  HeaderSize;
-	};
-
-	virtual void WritePackage(const FPackageInfo& Info, const FIoBuffer& PackageData, const TArray<FFileRegion>& FileRegions) = 0;
-
-	struct FBulkDataInfo
-	{
-		enum EType 
-		{
-			Standard,
-			Mmap,
-			Optional
-		};
-
-		FName	PackageName;
-		EType	BulkdataType = Standard;
-		FString	LooseFilePath;
-	};
-
-	virtual void WriteBulkdata(const FBulkDataInfo& Info, const FIoBuffer& BulkData, const TArray<FFileRegion>& FileRegions) = 0;
-
-	virtual void Finalize() = 0;
-};
-
 class FSavePackageContext
 {
 public:
-	FSavePackageContext(IPackageStoreWriter* InPackageStoreWriter, bool InbForceLegacyOffsets)
-	: PackageStoreWriter(InPackageStoreWriter) 
+	FSavePackageContext(const ITargetPlatform* InTargetPlatform, IPackageStoreWriter* InPackageStoreWriter, bool InbForceLegacyOffsets)
+	: TargetPlatform(InTargetPlatform)
+	, PackageStoreWriter(InPackageStoreWriter) 
 	, bForceLegacyOffsets(InbForceLegacyOffsets)
 	{
 	}
 
 	COREUOBJECT_API ~FSavePackageContext();
 
-	IPackageStoreWriter* PackageStoreWriter;
-	bool bForceLegacyOffsets;
+	const ITargetPlatform* const TargetPlatform;
+	IPackageStoreWriter* const PackageStoreWriter;
+	const bool bForceLegacyOffsets;
 };
 
 namespace UE
diff --git a/Engine/Source/Runtime/Engine/Private/ODSC/ODSCThread.cpp b/Engine/Source/Runtime/Engine/Private/ODSC/ODSCThread.cpp
--- a/Engine/Source/Runtime/Engine/Private/ODSC/ODSCThread.cpp
+++ b/Engine/Source/Runtime/Engine/Private/ODSC/ODSCThread.cpp
@@ -7,6 +7,7 @@
 #include "HAL/PlatformProcess.h"
 #include "HAL/RunnableThread.h"
 #include "ShaderCompiler.h"
+#include "CookOnTheFly.h"
 
 FODSCRequestPayload::FODSCRequestPayload(EShaderPlatform InShaderPlatform, const FString& InMaterialName, const FString& InVertexFactoryName, const FString& InPipelineName, const TArray<FString>& InShaderTypeNames, const FString& InRequestHash)
 	: ShaderPlatform(InShaderPlatform), MaterialName(InMaterialName), VertexFactoryName(InVertexFactoryName), PipelineName(InPipelineName), ShaderTypeNames(std::move(InShaderTypeNames)), RequestHash(InRequestHash)
@@ -197,7 +198,17 @@ void FODSCThread::Process()
 	for (FODSCMessageHandler* NextRequest : RequestsToStart)
 	{
 		// send the info, the handler will process the response (and update shaders, etc)
-		IFileManager::Get().SendMessageToServer(TEXT("RecompileShaders"), NextRequest);
+		UE::Cook::SendCookOnTheFlyRequest(
+			UE::Cook::ECookOnTheFlyMessage::RecompileShaders,
+			[NextRequest](FArchive& Request)
+			{
+				NextRequest->FillPayload(Request);
+			},
+			[NextRequest](FArchive& Response)
+			{
+				NextRequest->ProcessResponse(Response);
+				return true;
+			});
 
 		CompletedThreadedRequests.Enqueue(NextRequest);
 	}
@@ -205,18 +216,27 @@ void FODSCThread::Process()
 	// process any specific mesh material shader requests.
 	if (PayloadsToAggregate.Num())
 	{
-		FODSCMessageHandler* requestHandler = new FODSCMessageHandler(PayloadsToAggregate[0].ShaderPlatform);
+		FODSCMessageHandler* RequestHandler = new FODSCMessageHandler(PayloadsToAggregate[0].ShaderPlatform);
 		for (const FODSCRequestPayload& payload : PayloadsToAggregate)
 		{
-			requestHandler->AddPayload(payload);
+			RequestHandler->AddPayload(payload);
 		}
 
 		// send the info, the handler will process the response (and update shaders, etc)
-		IFileManager::Get().SendMessageToServer(TEXT("RecompileShaders"), requestHandler);
+		UE::Cook::SendCookOnTheFlyRequest(
+			UE::Cook::ECookOnTheFlyMessage::RecompileShaders,
+			[RequestHandler](FArchive& Request)
+			{
+				RequestHandler->FillPayload(Request);
+			},
+			[RequestHandler](FArchive& Response)
+			{
+				RequestHandler->ProcessResponse(Response);
+				return true;
+			});
 
-		CompletedThreadedRequests.Enqueue(requestHandler);
+		CompletedThreadedRequests.Enqueue(RequestHandler);
 	}
 
 	WakeupEvent->Reset();
 }
-
diff --git a/Engine/Source/Runtime/Launch/Launch.Build.cs b/Engine/Source/Runtime/Launch/Launch.Build.cs
--- a/Engine/Source/Runtime/Launch/Launch.Build.cs
+++ b/Engine/Source/Runtime/Launch/Launch.Build.cs
@@ -34,7 +34,8 @@ public class Launch : ModuleRules
 				"TraceLog",
 				"Overlay",
 				"PreLoadScreen",
-				"InstallBundleManager"
+				"InstallBundleManager",
+				"CookOnTheFly"
 			});
 
 		// Set a macro allowing us to switch between debuggame/development configuration
@@ -260,6 +261,11 @@ public class Launch : ModuleRules
 			PrivateDependencyModuleNames.Add("UnixCommonStartup");
 		}
 
+		if (Target.Configuration != UnrealTargetConfiguration.Shipping)
+		{
+			PrivateDependencyModuleNames.Add("StorageServerClient");
+		}
+
 		if(Target.LinkType == TargetLinkType.Monolithic && !Target.bFormalBuild)
 		{
 			PrivateDefinitions.Add(string.Format("COMPILED_IN_CL={0}", Target.Version.Changelist));
diff --git a/Engine/Source/Runtime/Launch/Private/LaunchEngineLoop.cpp b/Engine/Source/Runtime/Launch/Private/LaunchEngineLoop.cpp
--- a/Engine/Source/Runtime/Launch/Private/LaunchEngineLoop.cpp
+++ b/Engine/Source/Runtime/Launch/Private/LaunchEngineLoop.cpp
@@ -74,6 +74,10 @@
 #include "IO/IoDispatcher.h"
 #endif
 
+#if WITH_COTF
+#include "CookOnTheFly.h"
+#endif
+
 #if WITH_COREUOBJECT
 	#include "Internationalization/PackageLocalizationManager.h"
 	#include "Misc/PackageName.h"
@@ -795,7 +799,22 @@ bool LaunchCheckForFileOverride(const TCHAR* CmdLine, bool& OutFileOverrideFound
 	// Get the physical platform file.
 	IPlatformFile* CurrentPlatformFile = &FPlatformFileManager::Get().GetPlatformFile();
 
+	// Try to create storage server wrapper
+	bool bIsUsingStorageServer = false;
+#if !UE_BUILD_SHIPPING
+	{
+		IPlatformFile* PlatformFile = ConditionallyCreateFileWrapper(TEXT("StorageServerClient"), CurrentPlatformFile, CmdLine);
+		if (PlatformFile)
+		{
+			CurrentPlatformFile = PlatformFile;
+			FPlatformFileManager::Get().SetPlatformFile(*CurrentPlatformFile);
+			bIsUsingStorageServer = true;
+		}
+	}
+#endif
+
 	// Try to create pak file wrapper
+	if (!bIsUsingStorageServer)
 	{
 		IPlatformFile* PlatformFile = ConditionallyCreateFileWrapper(TEXT("PakFile"), CurrentPlatformFile, CmdLine);
 		if (PlatformFile)
@@ -812,6 +831,7 @@ bool LaunchCheckForFileOverride(const TCHAR* CmdLine, bool& OutFileOverrideFound
 	}
 
 	// Try to create sandbox wrapper
+	if (!bIsUsingStorageServer)
 	{
 		IPlatformFile* PlatformFile = ConditionallyCreateFileWrapper(TEXT("SandboxFile"), CurrentPlatformFile, CmdLine);
 		if (PlatformFile)
@@ -1793,6 +1813,24 @@ int32 FEngineLoop::PreInitPreStartupScreen(const TCHAR* CmdLine)
 	}
 #endif // WITH_ENGINE
 
+#if WITH_COTF
+	if (IsRunningCookOnTheFly())
+	{
+		UE::Cook::FCookOnTheFlyHostOptions CookOnTheFlyHostOptions;
+		if (GetCookOnTheFlyHost(CookOnTheFlyHostOptions))
+		{
+			using namespace UE::Cook;
+			ICookOnTheFlyModule& CookOnTheFlyModule = FModuleManager::LoadModuleChecked<ICookOnTheFlyModule>(TEXT("CookOnTheFly"));
+			const bool bConnected = CookOnTheFlyModule.ConnectToServer(CookOnTheFlyHostOptions);
+			if (!bConnected)
+			{
+				UE_LOG(LogInit, Error, TEXT("Failed to connect to cook on the fly server"));
+				return 1;
+			}
+		}
+	}
+#endif // WITH_COTF
+
 	// allow the command line to override the platform file singleton
 	bool bFileOverrideFound = false;
 	{
@@ -3913,7 +3951,7 @@ void FEngineLoop::LoadPreInitModules()
 	// Load audio editor module before engine class CDOs are loaded
 	FModuleManager::Get().LoadModule(TEXT("AudioEditor"));
 	FModuleManager::Get().LoadModule(TEXT("AnimationModifiers"));
-#endif
+#endif	
 }
 
 
diff --git a/Engine/Source/Runtime/RenderCore/Private/ShaderCodeLibrary.cpp b/Engine/Source/Runtime/RenderCore/Private/ShaderCodeLibrary.cpp
--- a/Engine/Source/Runtime/RenderCore/Private/ShaderCodeLibrary.cpp
+++ b/Engine/Source/Runtime/RenderCore/Private/ShaderCodeLibrary.cpp
@@ -2655,8 +2655,7 @@ void FShaderCodeLibrary::InitForRuntime(EShaderPlatform ShaderPlatform)
 	// We cannot enable native shader libraries when running with NullRHI, so for consistency all libraries (both native and non-native) are disabled if FApp::CanEverRender() == false
 	bool bEnable = !FPlatformProperties::IsServerOnly() && FApp::CanEverRender() && bArchive;
 #if !UE_BUILD_SHIPPING
-	FString FileHostIP;
-	const bool bCookOnTheFly = FParse::Value(FCommandLine::Get(), TEXT("filehostip"), FileHostIP);
+	const bool bCookOnTheFly = IsRunningCookOnTheFly(); 
 	bEnable &= !bCookOnTheFly;
 #endif
 
diff --git a/Engine/Source/Runtime/RenderCore/RenderCore.Build.cs b/Engine/Source/Runtime/RenderCore/RenderCore.Build.cs
--- a/Engine/Source/Runtime/RenderCore/RenderCore.Build.cs
+++ b/Engine/Source/Runtime/RenderCore/RenderCore.Build.cs
@@ -21,7 +21,7 @@ public class RenderCore : ModuleRules
 			PrivateDependencyModuleNames.Add("CoreUObject");
 		}
 
-        PrivateDependencyModuleNames.AddRange(new string[] { "Core", "Projects", "RHI", "ApplicationCore", "TraceLog" });
+        PrivateDependencyModuleNames.AddRange(new string[] { "Core", "Projects", "RHI", "ApplicationCore", "TraceLog", "CookOnTheFly" });
 
         PrivateIncludePathModuleNames.AddRange(new string[] { "DerivedDataCache" });
 		
diff --git a/Engine/Source/Runtime/StorageServerClient/Private/StorageServerConnection.cpp b/Engine/Source/Runtime/StorageServerClient/Private/StorageServerConnection.cpp
new file mode 100644
--- /dev/null
+++ b/Engine/Source/Runtime/StorageServerClient/Private/StorageServerConnection.cpp
@@ -0,0 +1,584 @@
+// Copyright Epic Games, Inc. All Rights Reserved.
+
+#include "StorageServerConnection.h"
+#include "IPAddress.h"
+#include "SocketSubsystem.h"
+#include "Sockets.h"
+#include "Misc/App.h"
+#include "Misc/Paths.h"
+#include "IO/IoDispatcher.h"
+#include "Misc/StringBuilder.h"
+#include "Misc/ScopeLock.h"
+#include "Serialization/CompactBinary.h"
+#include "Serialization/CompactBinarySerialization.h"
+
+#if !UE_BUILD_SHIPPING
+
+DEFINE_LOG_CATEGORY_STATIC(LogStorageServerConnection, Log, All);
+
+FStorageServerRequest::FStorageServerRequest(FAnsiStringView Verb, FAnsiStringView Resource, FAnsiStringView Hostname)
+{
+	SetIsSaving(true);
+	HeaderBuffer << Verb << " " << Resource << " HTTP/1.1\r\n" 
+		<< "Host: " << Hostname << "\r\n"
+		<< "Connection: Keep-Alive\r\n";
+}
+
+FSocket* FStorageServerRequest::Send(FStorageServerConnection& Owner)
+{
+	if (BodyBuffer.Num())
+	{
+		HeaderBuffer.Append("Content-Length: ").Appendf("%d\r\n", BodyBuffer.Num());
+	}
+	HeaderBuffer << "\r\n";
+	int32 BytesLeft = HeaderBuffer.Len();
+	
+	auto Send = [](FSocket* Socket, const uint8* Data, int32 Length)
+	{
+		int32 BytesLeft = Length;
+		while (BytesLeft > 0)
+		{
+			int32 BytesSent;
+			if (!Socket->Send(Data, BytesLeft, BytesSent))
+			{
+				return false;
+			}
+			check(BytesSent >= 0);
+			BytesLeft -= BytesSent;
+			Data += BytesSent;
+		}
+		return true;
+	};
+
+	int32 Attempts = 0;
+	while (Attempts++ < 10)
+	{
+		FSocket* Socket = Owner.AcquireSocket();
+		if (Send(Socket, reinterpret_cast<const uint8*>(HeaderBuffer.GetData()), HeaderBuffer.Len()) &&
+			Send(Socket, BodyBuffer.GetData(), BodyBuffer.Num()))
+		{
+			return Socket;
+		}
+		UE_LOG(LogStorageServerConnection, Warning, TEXT("Failed to send request to storage server. Retrying..."));
+		Owner.ReleaseSocket(Socket, false);
+	}
+	return nullptr;
+}
+
+void FStorageServerRequest::Serialize(void* V, int64 Length)
+{
+	int32 Index = BodyBuffer.AddUninitialized(Length);
+	uint8* Dest = BodyBuffer.GetData() + Index;
+	FMemory::Memcpy(Dest, V, Length);
+}
+
+FStorageServerResponse::FStorageServerResponse(FStorageServerConnection& InOwner, FSocket& InSocket)
+	: Owner(InOwner)
+	, Socket(&InSocket)
+{
+	SetIsLoading(true);
+	uint8 Buffer[1024];
+	int32 TotalReadFromSocket = 0;
+	auto ReadResponseLine = [&Buffer, &InSocket, &TotalReadFromSocket]() -> FAnsiStringView
+	{
+		for (;;)
+		{
+			int32 BytesRead;
+			InSocket.Recv(Buffer, 1024, BytesRead, ESocketReceiveFlags::Peek);
+			FAnsiStringView ResponseView(reinterpret_cast<const ANSICHAR*>(Buffer), BytesRead);
+			int32 LineEndIndex;
+			if (ResponseView.FindChar('\r', LineEndIndex) && BytesRead >= LineEndIndex + 2)
+			{
+				check(ResponseView[LineEndIndex + 1] == '\n');
+				InSocket.Recv(Buffer, LineEndIndex + 2, BytesRead, ESocketReceiveFlags::None);
+				check(BytesRead == LineEndIndex + 2);
+				TotalReadFromSocket += BytesRead;
+				return ResponseView.Left(LineEndIndex);
+			}
+		}
+	};
+
+	FAnsiStringView ResponseLine = ReadResponseLine();
+	if (ResponseLine == "HTTP/1.1 200 OK")
+	{
+		bIsOk = true;
+	}
+	else if (ResponseLine.StartsWith("HTTP/1.1 "))
+	{
+		ErrorCode = TCString<ANSICHAR>::Atoi64(ResponseLine.GetData() + 9);
+	}
+	while (!ResponseLine.IsEmpty())
+	{
+		ResponseLine = ReadResponseLine();
+		if (ResponseLine.StartsWith("Content-Length: "))
+		{
+			ContentLength = TCString<ANSICHAR>::Atoi64(ResponseLine.GetData() + 16);
+		}
+	}
+	if (!bIsOk && ContentLength)
+	{
+		TArray<uint8> ErrorBuffer;
+		ErrorBuffer.SetNumUninitialized(ContentLength);
+		int32 BytesRead;
+		InSocket.Recv(ErrorBuffer.GetData(), ContentLength, BytesRead, ESocketReceiveFlags::WaitAll);
+		ErrorMessage = FString(ContentLength, ANSI_TO_TCHAR(reinterpret_cast<ANSICHAR*>(ErrorBuffer.GetData())));
+		ContentLength = 0;
+	}
+	if (ContentLength == 0)
+	{
+		ReleaseSocket(true);
+	}
+}
+
+FStorageServerChunkBatchRequest::FStorageServerChunkBatchRequest(FStorageServerConnection& InOwner, FAnsiStringView Resource, FAnsiStringView Hostname)
+	: FStorageServerRequest("POST", Resource, Hostname)
+	, Owner(InOwner)
+{
+	uint32 Magic = 0xAAAA'77AC;
+	uint32 ChunkCountPlaceHolder = 0;
+	uint32 Reserved1 = 0;
+	uint32 Reserved2 = 0;
+	*this << Magic;
+	ChunkCountOffset = BodyBuffer.Num();
+	*this << ChunkCountPlaceHolder << Reserved1 << Reserved2;
+}
+
+FStorageServerChunkBatchRequest& FStorageServerChunkBatchRequest::AddChunk(const FIoChunkId& ChunkId, int64 Offset, int64 Size)
+{
+	uint32* ChunkCount = reinterpret_cast<uint32*>(BodyBuffer.GetData() + ChunkCountOffset);
+	*this << const_cast<FIoChunkId&>(ChunkId) << *ChunkCount << Offset << Size;
+	++(*ChunkCount);
+	return *this;
+}
+
+bool FStorageServerChunkBatchRequest::Issue(TFunctionRef<void(uint32 ChunkCount, uint32* ChunkIndices, uint64* ChunkSizes, FStorageServerResponse& ChunkDataStream)> OnResponse)
+{
+	FSocket* Socket = Send(Owner);
+	if (!Socket)
+	{
+		UE_LOG(LogStorageServerConnection, Fatal, TEXT("Failed to send chunk batch request to storage server."));
+		return false;
+	}
+	FStorageServerResponse Response(Owner, *Socket);
+	if (!Response.IsOk())
+	{
+		UE_LOG(LogStorageServerConnection, Fatal, TEXT("Failed to read chunk batch from storage server. '%s'"), *Response.GetErrorMessage());
+		return false;
+	}
+
+	uint32 Magic;
+	uint32 ChunkCount;
+	uint32 Reserved1;
+	uint32 Reserved2;
+	Response << Magic;
+	if (Magic != 0xbada'b00f)
+	{
+		UE_LOG(LogStorageServerConnection, Fatal, TEXT("Invalid magic in chunk batch response from storage server."));
+		return false;
+	}
+	Response << ChunkCount;
+	Response << Reserved1;
+	Response << Reserved2;
+
+	TArray<uint32, TInlineAllocator<64>> ChunkIndices;
+	ChunkIndices.Reserve(ChunkCount);
+	TArray<uint64, TInlineAllocator<64>> ChunkSizes;
+	ChunkSizes.Reserve(ChunkCount);
+	for (uint32 Index = 0; Index < ChunkCount; ++Index)
+	{
+		uint32 ChunkIndex;
+		uint32 Flags;
+		int64 ChunkSize;
+		Response << ChunkIndex;
+		Response << Flags;
+		Response << ChunkSize;
+		ChunkIndices.Add(ChunkIndex);
+		ChunkSizes.Emplace(ChunkSize);
+	}
+	OnResponse(ChunkCount, ChunkIndices.GetData(), ChunkSizes.GetData(), Response);
+	return true;
+}
+
+void FStorageServerResponse::ReleaseSocket(bool bKeepAlive)
+{
+	Owner.ReleaseSocket(Socket, bKeepAlive);
+	Socket = nullptr;
+}
+
+void FStorageServerResponse::Serialize(void* V, int64 Length)
+{
+	if (Length == 0)
+	{
+		return;
+	}
+	if (!Socket)
+	{
+		UE_LOG(LogStorageServerConnection, Fatal, TEXT("Trying to read %lld bytes from released socket"), Length);
+		return;
+	}
+	if (Position + Length > ContentLength)
+	{
+		UE_LOG(LogStorageServerConnection, Fatal, TEXT("Trying to read %lld bytes from socket with only %lld available"), Length, ContentLength - Position);
+		return;
+	}
+	uint64 RemainingBytesToRead = Length;
+	uint8* Destination = reinterpret_cast<uint8*>(V);
+	while (RemainingBytesToRead)
+	{
+		uint64 BytesToRead32 = FMath::Min(RemainingBytesToRead, static_cast<uint64>(INT32_MAX));
+		int32 BytesRead;
+		if (!Socket->Recv(Destination, static_cast<int32>(BytesToRead32), BytesRead, ESocketReceiveFlags::WaitAll))
+		{
+			UE_LOG(LogStorageServerConnection, Fatal, TEXT("Failed reading %d bytes from socket"), BytesToRead32);
+			return;
+		}
+		RemainingBytesToRead -= BytesRead;
+		Destination += BytesRead;
+		Position += BytesRead;
+	}
+	
+	if (Position == ContentLength)
+	{
+		ReleaseSocket(true);
+	}
+}
+
+FStorageServerConnection::FStorageServerConnection()
+	: SocketSubsystem(*ISocketSubsystem::Get())
+{
+}
+
+FStorageServerConnection::~FStorageServerConnection()
+{
+	for (FSocket* Socket : SocketPool)
+	{
+		Socket->Close();
+		delete Socket;
+	}
+}
+
+bool FStorageServerConnection::Initialize(const TCHAR* InHost, int32 InPort, const TCHAR* InProjectNameOverride, const TCHAR* InPlatformNameOverride)
+{
+	FAddressInfoResult GAIRequest = SocketSubsystem.GetAddressInfo(InHost, nullptr, EAddressInfoFlags::Default, NAME_None);
+	if (GAIRequest.ReturnCode == SE_NO_ERROR && GAIRequest.Results.Num() > 0)
+	{
+		ServerAddr = GAIRequest.Results[0].Address;
+	}
+	if (!ServerAddr.IsValid() || !ServerAddr->IsValid())
+	{
+		UE_LOG(LogStorageServerConnection, Fatal, TEXT("Failed to resolve storage server host %s."), InHost);
+		return false;
+	}
+	
+	ServerAddr->SetPort(InPort);
+	OplogPath.Append("/prj/");
+	if (InProjectNameOverride)
+	{
+		OplogPath.Append(TCHAR_TO_ANSI(InProjectNameOverride));
+	}
+	else
+	{
+		OplogPath.Append(TCHAR_TO_ANSI(FApp::GetProjectName()));
+	}
+	OplogPath.Append("/oplog/");
+	if (InPlatformNameOverride)
+	{
+		OplogPath.Append(TCHAR_TO_ANSI(InPlatformNameOverride));
+	}
+	else
+	{
+		OplogPath.Append(FPlatformProperties::PlatformName());
+	}
+
+	Hostname.Append(TCHAR_TO_ANSI(*ServerAddr->ToString(false)));
+
+	int32 ServerVersion = HandshakeRequest();
+	if (ServerVersion != 1)
+	{
+		return false;
+	}
+
+	UE_LOG(LogStorageServerConnection, Display, TEXT("Connected to storage server host %s."), *ServerAddr->ToString(true));
+
+	return true;
+}
+
+int32 FStorageServerConnection::HandshakeRequest()
+{
+	TAnsiStringBuilder<256> ResourceBuilder;
+	ResourceBuilder.Append(OplogPath);
+	FStorageServerRequest Request("GET", *ResourceBuilder, Hostname);
+	FSocket* Socket = Request.Send(*this);
+	if (!Socket)
+	{
+		UE_LOG(LogStorageServerConnection, Fatal, TEXT("Failed to send handshake request to storage server at %s."), *ServerAddr->ToString(true));
+		return -1;
+	}
+	FStorageServerResponse Response(*this, *Socket);
+
+	if (Response.IsOk())
+	{
+		FCbObject ResponseObj = Response.GetResponseObject();
+
+		// we currently don't have any concept of protocol versioning, if
+		// we succeed in communicating with the endpoint we're good since
+		// any breaking API change would need to be done in a backward
+		// compatible manner
+
+		return 1;
+	}
+	else
+	{
+		UE_LOG(LogStorageServerConnection, Fatal, TEXT("Failed to handshake with storage server at %s. '%s'"), *ServerAddr->ToString(true), *Response.GetErrorMessage());
+	}
+	return -1;
+}
+
+void FStorageServerConnection::FileManifestRequest(TFunctionRef<void(FIoChunkId Id, FStringView Path)> Callback)
+{
+	TAnsiStringBuilder<256> ResourceBuilder;
+	ResourceBuilder.Append(OplogPath).Append("/files");
+	FStorageServerRequest Request("GET", *ResourceBuilder, Hostname);
+	FSocket* Socket = Request.Send(*this);
+	if (!Socket)
+	{
+		UE_LOG(LogStorageServerConnection, Fatal, TEXT("Failed to send file manifest request to storage server at %s."), *ServerAddr->ToString(true));
+		return;
+	}
+
+	FStorageServerResponse Response(*this, *Socket);
+	if (Response.IsOk())
+	{
+		FCbObject ResponseObj = Response.GetResponseObject();
+		
+		for (FCbField& FileArrayEntry : ResponseObj["files"].AsArray())
+		{
+			FCbObject Entry = FileArrayEntry.AsObject();
+			FCbObjectId Id = Entry["id"].AsObjectId();
+
+			TStringBuilder<128> WidePath;
+			WidePath.Append(FUTF8ToTCHAR(Entry["path"].AsString()));
+
+			FIoChunkId ChunkId;
+			ChunkId.Set(Id.GetView());
+
+			Callback(ChunkId, WidePath);
+		}
+	}
+	else
+	{
+		UE_LOG(LogStorageServerConnection, Fatal, TEXT("Failed to read file manifest from storage server at %s. '%s'"), *ServerAddr->ToString(true), *Response.GetErrorMessage());
+	}
+}
+
+void FileIndexToChunkId(FAnsiStringBuilderBase& OutString, int32 FileIndex)
+{
+	FIoChunkId Chunk = CreateExternalFileChunkId(0, FileIndex);
+	OutString << Chunk;
+}
+
+int64 FStorageServerConnection::FileSizeRequest(int32 FileIndex)
+{
+	TAnsiStringBuilder<256> ResourceBuilder;
+	ResourceBuilder.Append(OplogPath);
+	ResourceBuilder << "/";
+	FileIndexToChunkId(ResourceBuilder, FileIndex);
+
+	FStorageServerRequest Request("HEAD", *ResourceBuilder, Hostname);
+	FSocket* Socket = Request.Send(*this);
+	if (!Socket)
+	{
+		UE_LOG(LogStorageServerConnection, Fatal, TEXT("Failed to send file size request to storage server at %s."), *ServerAddr->ToString(true));
+		return -1;
+	}
+	FStorageServerResponse Response(*this, *Socket);
+	if (Response.IsOk())
+	{
+		return Response.ContentLength;
+	}
+	else if (Response.GetErrorCode() == 404)
+	{
+		return -1;
+	}
+	else
+	{
+		UE_LOG(LogStorageServerConnection, Fatal, TEXT("Failed to get file size from storage server at %s. '%s'"), *ServerAddr->ToString(true), *Response.GetErrorMessage());
+	}
+	return -1;
+}
+
+bool FStorageServerConnection::ReadFileRequest(int32 FileIndex, uint64 Offset, uint64 Size, TFunctionRef<void(FStorageServerResponse&)> OnResponse)
+{
+	TAnsiStringBuilder<256> ResourceBuilder;
+	ResourceBuilder.Append(OplogPath);
+	ResourceBuilder << "/";
+	FileIndexToChunkId(ResourceBuilder, FileIndex);
+
+	ResourceBuilder.Appendf("?offset=%" UINT64_FMT "&size=%" UINT64_FMT, Offset, Size);
+
+	FStorageServerRequest Request = FStorageServerRequest("GET", *ResourceBuilder, Hostname);
+	FSocket* Socket = Request.Send(*this);
+	if (!Socket)
+	{
+		UE_LOG(LogStorageServerConnection, Fatal, TEXT("Failed to send file read request to storage server at %s."), *ServerAddr->ToString(true));
+		return false;
+	}
+	FStorageServerResponse Response(*this, *Socket);
+	if (Response.IsOk())
+	{
+		OnResponse(Response);
+		return true;
+	}
+	else if (Response.GetErrorCode() == 404)
+	{
+		return false;
+	}
+	else
+	{
+		UE_LOG(LogStorageServerConnection, Fatal, TEXT("Failed to read file from storage server at %s."), *ServerAddr->ToString(true), *Response.GetErrorMessage());
+	}
+	return false;
+}
+
+int64 FStorageServerConnection::ChunkSizeRequest(const FIoChunkId& ChunkId)
+{
+	TAnsiStringBuilder<256> ResourceBuilder;
+	ResourceBuilder.Append(OplogPath);
+	ResourceBuilder << "/" << ChunkId << "/info";
+
+	FStorageServerRequest Request("GET", *ResourceBuilder, Hostname);
+	FSocket* Socket = Request.Send(*this);
+	if (!Socket)
+	{
+		UE_LOG(LogStorageServerConnection, Fatal, TEXT("Failed to send chunk size request to storage server at %s."), *ServerAddr->ToString(true));
+		return -1;
+	}
+	FStorageServerResponse Response(*this, *Socket);
+	if (Response.IsOk())
+	{
+		FCbObject ResponseObj = Response.GetResponseObject();
+
+		const int64 ChunkSize = ResponseObj["size"].AsInt64(0);
+
+		return ChunkSize;
+	}
+	else if (Response.GetErrorCode() == 404)
+	{
+		return -1;
+	}
+	else
+	{
+		UE_LOG(LogStorageServerConnection, Fatal, TEXT("Failed to get chunk size from storage server at %s. '%s'"), *ServerAddr->ToString(true), *Response.GetErrorMessage());
+	}
+	return -1;
+}
+
+bool FStorageServerConnection::ReadChunkRequest(const FIoChunkId& ChunkId, uint64 Offset, uint64 Size, TFunctionRef<void(FStorageServerResponse&)> OnResponse)
+{
+	TAnsiStringBuilder<256> ResourceBuilder;
+	ResourceBuilder.Append(OplogPath) << "/" << ChunkId;
+
+	bool HaveQuery = false;
+
+	auto AppendQueryDelimiter = [&] 
+	{
+		if (HaveQuery)
+		{
+			ResourceBuilder.Append("&"_ASV);
+		}
+		else
+		{
+			ResourceBuilder.Append("?"_ASV);
+			HaveQuery = true;
+		}
+	};
+
+	if (Offset)
+	{
+		AppendQueryDelimiter();
+		ResourceBuilder.Appendf("offset=%" UINT64_FMT, Offset);
+	}
+
+	if (Size != ~uint64(0))
+	{
+		AppendQueryDelimiter();
+		ResourceBuilder.Appendf("size=%" UINT64_FMT, Size);
+	}
+
+	FStorageServerRequest Request("GET", *ResourceBuilder, Hostname);
+	FSocket* Socket = Request.Send(*this);
+	if (!Socket)
+	{
+		UE_LOG(LogStorageServerConnection, Fatal, TEXT("Failed to send chunk read request to storage server at %s."), *ServerAddr->ToString(true));
+		return false;
+	}
+	FStorageServerResponse Response(*this, *Socket);
+	if (Response.IsOk())
+	{
+		OnResponse(Response);
+		return true;
+	}
+	else if (Response.GetErrorCode() == 404)
+	{
+		return false;
+	}
+	else
+	{
+		UE_LOG(LogStorageServerConnection, Fatal, TEXT("Failed to read chunk from storage server at %s. '%s'"), *ServerAddr->ToString(true), *Response.GetErrorMessage());
+		return false;
+	}
+}
+
+FStorageServerChunkBatchRequest FStorageServerConnection::NewChunkBatchRequest()
+{
+	TAnsiStringBuilder<256> ResourceBuilder;
+	ResourceBuilder.Append(OplogPath).Append("/batch");
+	return FStorageServerChunkBatchRequest(*this, *ResourceBuilder, Hostname);
+}
+
+FSocket* FStorageServerConnection::AcquireSocket()
+{
+	{
+		FScopeLock Lock(&SocketPoolCritical);
+		if (!SocketPool.IsEmpty())
+		{
+			return SocketPool.Pop(false);
+		}
+	}
+
+	for (int32 Attempt = 0, MaxAttempts = 10; Attempt < MaxAttempts; Attempt++)
+	{
+		FSocket* Socket = SocketSubsystem.CreateSocket(NAME_Stream, TEXT("StorageServer"), ServerAddr->GetProtocolType());
+		check(Socket);
+
+		if (Socket->Connect(*ServerAddr))
+		{
+			return Socket;
+		}
+
+		delete Socket;
+	}
+
+	UE_LOG(LogStorageServerConnection, Fatal, TEXT("Failed to connect to storage server at %s."), *ServerAddr->ToString(true));
+
+	return nullptr;
+}
+
+void FStorageServerConnection::ReleaseSocket(FSocket* Socket, bool bKeepAlive)
+{
+	if (bKeepAlive)
+	{
+		uint32 PendingDataSize;
+		if (!Socket->HasPendingData(PendingDataSize))
+		{
+			FScopeLock Lock(&SocketPoolCritical);
+			SocketPool.Push(Socket);
+			return;
+		}
+		UE_LOG(LogStorageServerConnection, Fatal, TEXT("Socket was not fully drained"));
+	}
+	Socket->Close();
+	delete Socket;
+}
+
+#endif
diff --git a/Engine/Source/Runtime/StorageServerClient/Private/StorageServerPlatformFile.cpp b/Engine/Source/Runtime/StorageServerClient/Private/StorageServerPlatformFile.cpp
new file mode 100644
--- /dev/null
+++ b/Engine/Source/Runtime/StorageServerClient/Private/StorageServerPlatformFile.cpp
@@ -0,0 +1,726 @@
+// Copyright Epic Games, Inc. All Rights Reserved.
+
+#include "StorageServerPlatformFile.h"
+#include "StorageServerIoDispatcherBackend.h"
+#include "HAL/IPlatformFileModule.h"
+#include "Misc/Paths.h"
+#include "Misc/ScopeRWLock.h"
+#include "StorageServerConnection.h"
+#include "Modules/ModuleManager.h"
+#include "Misc/StringBuilder.h"
+#include "Algo/Replace.h"
+
+#if WITH_COTF
+#include "Modules/ModuleManager.h"
+#include "CookOnTheFly.h"
+#endif
+
+DEFINE_LOG_CATEGORY_STATIC(LogStorageServerPlatformFile, Log, All);
+
+#if !UE_BUILD_SHIPPING
+
+FStorageServerFileSystemTOC::~FStorageServerFileSystemTOC()
+{
+	FWriteScopeLock _(TocLock);
+	for (auto& KV : Directories)
+	{
+		delete KV.Value;
+	}
+}
+
+FStorageServerFileSystemTOC::FDirectory* FStorageServerFileSystemTOC::AddDirectoriesRecursive(const FString& DirectoryPath)
+{
+	FDirectory* Directory = new FDirectory();
+	Directories.Add(DirectoryPath, Directory);
+	FString ParentDirectoryPath = FPaths::GetPath(DirectoryPath);
+	FDirectory* ParentDirectory;
+	if (ParentDirectoryPath.IsEmpty())
+	{
+		ParentDirectory = &Root;
+	}
+	else
+	{
+		ParentDirectory = Directories.FindRef(ParentDirectoryPath);
+		if (!ParentDirectory)
+		{
+			ParentDirectory = AddDirectoriesRecursive(ParentDirectoryPath);
+		}
+	}
+	ParentDirectory->Directories.Add(DirectoryPath);
+	return Directory;
+}
+
+void FStorageServerFileSystemTOC::AddFile(FStringView PathView, int32 Index)
+{
+	FWriteScopeLock _(TocLock);
+
+	FString Path(PathView);
+	FilePathToIndexMap.Add(Path, Index);
+	FileIndexToPathMap.Add(Index, Path);
+	FString DirectoryPath = FPaths::GetPath(Path);
+	FDirectory* Directory = Directories.FindRef(DirectoryPath);
+	if (!Directory)
+	{
+		Directory = AddDirectoriesRecursive(DirectoryPath);
+	}
+	Directory->Files.Add(Index);
+}
+
+bool FStorageServerFileSystemTOC::FileExists(const FString& Path)
+{
+	FReadScopeLock _(TocLock);
+	return FilePathToIndexMap.Contains(Path);
+}
+
+bool FStorageServerFileSystemTOC::DirectoryExists(const FString& Path)
+{
+	FReadScopeLock _(TocLock);
+	return Directories.Contains(Path);
+}
+
+int32* FStorageServerFileSystemTOC::FindFileIndex(const FString& Path)
+{
+	FReadScopeLock _(TocLock);
+	return FilePathToIndexMap.Find(Path);
+}
+
+bool FStorageServerFileSystemTOC::IterateDirectory(const FString& Path, TFunctionRef<bool(int32, const TCHAR*)> Callback)
+{
+	UE_LOG(LogStorageServerPlatformFile, Verbose, TEXT("IterateDirectory '%s'"), *Path);
+
+	FReadScopeLock _(TocLock);
+
+	FDirectory* Directory = Directories.FindRef(Path);
+	if (!Directory)
+	{
+		return false;
+	}
+	for (int32 FileIndex : Directory->Files)
+	{
+		const FString* FindFilePath = FileIndexToPathMap.Find(FileIndex);
+		if (FindFilePath && !Callback(FileIndex, **FindFilePath))
+		{
+			return false;
+		}
+	}
+	for (const FString& ChildDirectoryPath : Directory->Directories)
+	{
+		if (!Callback(-1, *ChildDirectoryPath))
+		{
+			return false;
+		}
+	}
+	return true;
+}
+
+class FStorageServerFileHandle
+	: public IFileHandle
+{
+	enum
+	{
+		BufferSize = 64 << 10
+	};
+	FStorageServerPlatformFile& Owner;
+	FString Filename;
+	int32 FileIndex;
+	int64 FilePos = 0;
+	int64 FileSize = -1;
+	int64 BufferStart = -1;
+	int64 BufferEnd = -1;
+	uint8 Buffer[BufferSize];
+
+public:
+	FStorageServerFileHandle(FStorageServerPlatformFile& InOwner, const TCHAR* InFilename, int32 InFileIndex)
+		: Owner(InOwner)
+		, Filename(InFilename)
+		, FileIndex(InFileIndex)
+	{
+	}
+
+	~FStorageServerFileHandle()
+	{
+	}
+
+	virtual int64 Size() override
+	{
+		if (FileSize < 0)
+		{
+			FFileStatData FileStatData = Owner.SendGetStatDataMessage(FileIndex);
+			if (FileStatData.bIsValid)
+			{
+				FileSize = FileStatData.FileSize;
+			}
+			else
+			{
+				UE_LOG(LogStorageServerPlatformFile, Warning, TEXT("Failed to obtain size of file '%s'"), *Filename);
+				FileSize = 0;
+			}
+		}
+		return FileSize;
+	}
+
+	virtual int64 Tell() override
+	{
+		return FilePos;
+	}
+
+	virtual bool Seek(int64 NewPosition) override
+	{
+		FilePos = NewPosition;
+		return true;
+	}
+
+	virtual bool SeekFromEnd(int64 NewPositionRelativeToEnd = 0) override
+	{
+		return Seek(Size() + NewPositionRelativeToEnd);
+	}
+
+	virtual bool Read(uint8* Destination, int64 BytesToRead) override
+	{
+		if (BytesToRead == 0)
+		{
+			return true;
+		}
+
+		if (BytesToRead > BufferSize)
+		{
+			int64 BytesRead = Owner.SendReadMessage(Destination, FileIndex, FilePos, BytesToRead);
+			if (BytesRead == BytesToRead)
+			{
+				FilePos += BytesRead;
+				return true;
+			}
+		}
+
+		if (FilePos < BufferStart || BufferEnd < FilePos + BytesToRead)
+		{
+			int64 BytesRead = Owner.SendReadMessage(Buffer, FileIndex, FilePos, BufferSize);
+			BufferStart = FilePos;
+			BufferEnd = BufferStart + BytesRead;
+		}
+
+		int64 BufferOffset = FilePos - BufferStart;
+		check(BufferEnd > BufferOffset);
+		int64 BytesToReadFromBuffer = FMath::Min(BufferEnd - BufferOffset, BytesToRead);
+		FMemory::Memcpy(Destination, Buffer + BufferOffset, BytesToReadFromBuffer);
+		if (BytesToReadFromBuffer == BytesToRead)
+		{
+			FilePos += BytesToReadFromBuffer;
+			return true;
+		}
+		
+		return false;
+	}
+
+	virtual bool Write(const uint8* Source, int64 BytesToWrite) override
+	{
+		check(false);
+		return false;
+	}
+
+	virtual bool Flush(const bool bFullFlush = false) override
+	{
+		return false;
+	}
+
+	virtual bool Truncate(int64 NewSize) override
+	{
+		return false;
+	}
+};
+
+FStorageServerPlatformFile::FStorageServerPlatformFile()
+{
+}
+
+FStorageServerPlatformFile::~FStorageServerPlatformFile()
+{
+}
+
+bool FStorageServerPlatformFile::ShouldBeUsed(IPlatformFile* Inner, const TCHAR* CmdLine) const
+{
+	FString StorageServerHost;
+	return FParse::Value(CmdLine, TEXT("-StorageServerHost="), StorageServerHost);
+}
+
+bool FStorageServerPlatformFile::Initialize(IPlatformFile* Inner, const TCHAR* CmdLine)
+{
+	LowerLevel = Inner;
+	bool bResult = false;
+	FString StorageServerHost;
+	if (FParse::Value(CmdLine, TEXT("-StorageServerHost="), StorageServerHost))
+	{
+		Connection.Reset(new FStorageServerConnection());
+		FString StorageServerProject;
+		FParse::Value(CmdLine, TEXT("-StorageServerProject="), StorageServerProject);
+		const TCHAR* ProjectOverride = StorageServerProject.IsEmpty() ? nullptr : *StorageServerProject;
+		FString StorageServerPlatform;
+		FParse::Value(CmdLine, TEXT("-StorageServerPlatform="), StorageServerPlatform);
+		const TCHAR* PlatformOverride = StorageServerPlatform.IsEmpty() ? nullptr : *StorageServerPlatform;
+		if (Connection->Initialize(*StorageServerHost, 1337, ProjectOverride, PlatformOverride))
+		{
+			if (SendGetFileListMessage())
+			{
+				FIoStatus IoDispatcherInitStatus = FIoDispatcher::Initialize();
+				if (IoDispatcherInitStatus.IsOk())
+				{
+					FIoDispatcher& IoDispatcher = FIoDispatcher::Get();
+					TSharedRef<FStorageServerIoDispatcherBackend> IoDispatcherBackend = MakeShared<FStorageServerIoDispatcherBackend>(*Connection.Get());
+					IoDispatcher.Mount(IoDispatcherBackend);
+#if WITH_COTF
+					if (IsRunningCookOnTheFly())
+					{
+						UE::Cook::ICookOnTheFlyModule& CookOnTheFlyModule = FModuleManager::LoadModuleChecked<UE::Cook::ICookOnTheFlyModule>(TEXT("CookOnTheFly"));
+						CookOnTheFlyModule.GetServerConnection().OnMessage().AddRaw(this, &FStorageServerPlatformFile::OnCookOnTheFlyMessage);
+					}
+#endif					
+					bResult = true;
+				}
+				else
+				{
+					UE_LOG(LogStorageServerPlatformFile, Fatal, TEXT("Failed to initialize IoDispatcher"), *StorageServerHost);
+				}
+			}
+			else
+			{
+				UE_LOG(LogStorageServerPlatformFile, Fatal, TEXT("Failed to get file list"), *StorageServerHost);
+			}
+		}
+		else
+		{
+			UE_LOG(LogStorageServerPlatformFile, Fatal, TEXT("Failed to initialize connection to %s"), *StorageServerHost);
+		}
+	}
+	if (!bResult)
+	{
+		Connection.Reset(nullptr);
+		return false;
+	}
+	return true;
+}
+
+bool FStorageServerPlatformFile::FileExists(const TCHAR* Filename)
+{
+	TStringBuilder<1024> StorageServerFilename;
+	if (MakeStorageServerPath(Filename, StorageServerFilename) && ServerToc.FileExists(*StorageServerFilename))
+	{
+		return true;
+	}
+	return LowerLevel->FileExists(Filename);
+}
+
+FDateTime FStorageServerPlatformFile::GetTimeStamp(const TCHAR* Filename)
+{
+	TStringBuilder<1024> StorageServerFilename;
+	if (MakeStorageServerPath(Filename, StorageServerFilename))
+	{
+		int32* FindFileIndex = ServerToc.FindFileIndex(*StorageServerFilename);
+		if (FindFileIndex)
+		{
+			FFileStatData FileStatData = SendGetStatDataMessage(*FindFileIndex);
+			check(FileStatData.bIsValid);
+			return FileStatData.ModificationTime;
+		}
+	}
+	return LowerLevel->GetTimeStamp(Filename);
+}
+
+FDateTime FStorageServerPlatformFile::GetAccessTimeStamp(const TCHAR* Filename)
+{
+	TStringBuilder<1024> StorageServerFilename;
+	if (MakeStorageServerPath(Filename, StorageServerFilename))
+	{
+		int32* FindFileIndex = ServerToc.FindFileIndex(*StorageServerFilename);
+		if (FindFileIndex)
+		{
+			FFileStatData FileStatData = SendGetStatDataMessage(*FindFileIndex);
+			check(FileStatData.bIsValid);
+			return FileStatData.AccessTime;
+		}
+	}
+	return LowerLevel->GetAccessTimeStamp(Filename);
+}
+
+int64 FStorageServerPlatformFile::FileSize(const TCHAR* Filename)
+{
+	TStringBuilder<1024> StorageServerFilename;
+	if (MakeStorageServerPath(Filename, StorageServerFilename))
+	{
+		int32* FindFileIndex = ServerToc.FindFileIndex(*StorageServerFilename);
+		if (FindFileIndex)
+		{
+			FFileStatData FileStatData = SendGetStatDataMessage(*FindFileIndex);
+			check(FileStatData.bIsValid);
+			return FileStatData.FileSize;
+		}
+	}
+	return LowerLevel->FileSize(Filename);
+}
+
+bool FStorageServerPlatformFile::IsReadOnly(const TCHAR* Filename)
+{
+	TStringBuilder<1024> StorageServerFilename;
+	if (MakeStorageServerPath(Filename, StorageServerFilename) && ServerToc.FileExists(*StorageServerFilename))
+	{
+		return true;
+	}
+	return LowerLevel->IsReadOnly(Filename);
+}
+
+FFileStatData FStorageServerPlatformFile::GetStatData(const TCHAR* FilenameOrDirectory)
+{
+	TStringBuilder<1024> StorageServerFilenameOrDirectory;
+	if (MakeStorageServerPath(FilenameOrDirectory, StorageServerFilenameOrDirectory))
+	{
+		int32* FindFileIndex = ServerToc.FindFileIndex(*StorageServerFilenameOrDirectory);
+		if (FindFileIndex)
+		{
+			return SendGetStatDataMessage(*FindFileIndex);
+		}
+		else if (ServerToc.DirectoryExists(*StorageServerFilenameOrDirectory))
+		{
+			return FFileStatData(
+				FDateTime::MinValue(),
+				FDateTime::MinValue(),
+				FDateTime::MinValue(),
+				0,
+				true,
+				true);
+		}
+	}
+	return LowerLevel->GetStatData(FilenameOrDirectory);
+}
+
+IFileHandle* FStorageServerPlatformFile::InternalOpenFile(int32 FileIndex, const TCHAR* LocalFilename)
+{
+	return new FStorageServerFileHandle(*this, LocalFilename, FileIndex);
+}
+
+IFileHandle* FStorageServerPlatformFile::OpenRead(const TCHAR* Filename, bool bAllowWrite)
+{
+	TStringBuilder<1024> StorageServerFilename;
+	if (MakeStorageServerPath(Filename, StorageServerFilename))
+	{
+		int32* FindFileIndex = ServerToc.FindFileIndex(*StorageServerFilename);
+		if (FindFileIndex)
+		{
+			return InternalOpenFile(*FindFileIndex, Filename);
+		}
+	}
+	return LowerLevel->OpenRead(Filename, bAllowWrite);
+}
+
+bool FStorageServerPlatformFile::IterateDirectory(const TCHAR* Directory, IPlatformFile::FDirectoryVisitor& Visitor)
+{
+	TStringBuilder<1024> StorageServerDirectory;
+	bool bResult = false;
+	if (MakeStorageServerPath(Directory, StorageServerDirectory))
+	{
+		bResult |= ServerToc.IterateDirectory(*StorageServerDirectory, [this, &Visitor](int32 FileIndex, const TCHAR* FilenameOrDirectory)
+		{
+			TStringBuilder<1024> LocalPath;
+			bool bConverted = MakeLocalPath(FilenameOrDirectory, LocalPath);
+			check(bConverted);
+			return Visitor.Visit(*LocalPath, FileIndex < 0);
+		});
+	}
+	bResult |= LowerLevel->IterateDirectory(Directory, Visitor);
+	return bResult;
+}
+
+bool FStorageServerPlatformFile::IterateDirectoryStat(const TCHAR* Directory, FDirectoryStatVisitor& Visitor)
+{
+	TStringBuilder<1024> StorageServerDirectory;
+	bool bResult = false;
+	if (MakeStorageServerPath(Directory, StorageServerDirectory))
+	{
+		bResult |= ServerToc.IterateDirectory(*StorageServerDirectory, [this, &Visitor](int32 FileIndex, const TCHAR* ServerFilenameOrDirectory)
+		{
+			TStringBuilder<1024> LocalPath;
+			bool bConverted = MakeLocalPath(ServerFilenameOrDirectory, LocalPath);
+			check(bConverted);
+			FFileStatData FileStatData;
+			if (FileIndex >= 0)
+			{
+				FileStatData = SendGetStatDataMessage(FileIndex);
+				check(FileStatData.bIsValid);
+			}
+			else
+			{
+				FileStatData = FFileStatData(
+					FDateTime::MinValue(),
+					FDateTime::MinValue(),
+					FDateTime::MinValue(),
+					0,
+					true,
+					true);
+			}
+			return Visitor.Visit(*LocalPath, FileStatData);
+		});
+	}
+	bResult |= LowerLevel->IterateDirectoryStat(Directory, Visitor);
+	return bResult;
+}
+
+bool FStorageServerPlatformFile::DirectoryExists(const TCHAR* Directory)
+{
+	TStringBuilder<1024> StorageServerDirectory;
+	if (MakeStorageServerPath(Directory, StorageServerDirectory) && ServerToc.DirectoryExists(*StorageServerDirectory))
+	{
+		return true;
+	}
+	return LowerLevel->DirectoryExists(Directory);
+}
+
+FString FStorageServerPlatformFile::GetFilenameOnDisk(const TCHAR* Filename)
+{
+	TStringBuilder<1024> StorageServerFilename;
+	if (MakeStorageServerPath(Filename, StorageServerFilename) && ServerToc.FileExists(*StorageServerFilename))
+	{
+		UE_LOG(LogStorageServerPlatformFile, Warning, TEXT("Attempting to get disk filename of remote file '%s'"), Filename);
+		return Filename;
+	}
+	return LowerLevel->GetFilenameOnDisk(Filename);
+}
+
+bool FStorageServerPlatformFile::DeleteFile(const TCHAR* Filename)
+{
+	TStringBuilder<1024> StorageServerFilename;
+	if (MakeStorageServerPath(Filename, StorageServerFilename) && ServerToc.FileExists(*StorageServerFilename))
+	{
+		return false;
+	}
+	return LowerLevel->DeleteFile(Filename);
+}
+
+bool FStorageServerPlatformFile::MoveFile(const TCHAR* To, const TCHAR* From)
+{
+	TStringBuilder<1024> StorageServerTo;
+	if (MakeStorageServerPath(To, StorageServerTo) && ServerToc.FileExists(*StorageServerTo))
+	{
+		return false;
+	}
+	TStringBuilder<1024> StorageServerFrom;
+	if (MakeStorageServerPath(From, StorageServerFrom))
+	{
+		int32* FindFromFileIndex = ServerToc.FindFileIndex(*StorageServerFrom);
+		if (FindFromFileIndex)
+		{
+			TUniquePtr<IFileHandle> ToFile(LowerLevel->OpenWrite(To, false, false));
+			if (!ToFile)
+			{
+				return false;
+			}
+
+			TUniquePtr<IFileHandle> FromFile(InternalOpenFile(*FindFromFileIndex, *StorageServerFrom));
+			if (!FromFile)
+			{
+				return false;
+			}
+			const int64 BufferSize = 64 << 10;
+			TArray<uint8> Buffer;
+			Buffer.SetNum(BufferSize);
+			int64 BytesLeft = FromFile->Size();
+			while (BytesLeft)
+			{
+				int64 BytesToWrite = FMath::Min(BufferSize, BytesLeft);
+				if (!FromFile->Read(Buffer.GetData(), BytesToWrite))
+				{
+					return false;
+				}
+				if (!ToFile->Write(Buffer.GetData(), BytesToWrite))
+				{
+					return false;
+				}
+				BytesLeft -= BytesToWrite;
+			}
+			return true;
+		}
+	}
+	return LowerLevel->MoveFile(To, From);
+}
+
+bool FStorageServerPlatformFile::SetReadOnly(const TCHAR* Filename, bool bNewReadOnlyValue)
+{
+	TStringBuilder<1024> StorageServerFilename;
+	if (MakeStorageServerPath(Filename, StorageServerFilename) && ServerToc.FileExists(*StorageServerFilename))
+	{
+		return bNewReadOnlyValue;
+	}
+	return LowerLevel->SetReadOnly(Filename, bNewReadOnlyValue);
+}
+
+void FStorageServerPlatformFile::SetTimeStamp(const TCHAR* Filename, FDateTime DateTime)
+{
+	TStringBuilder<1024> StorageServerFilename;
+	if (MakeStorageServerPath(Filename, StorageServerFilename) && ServerToc.FileExists(*StorageServerFilename))
+	{
+		return;
+	}
+	LowerLevel->SetTimeStamp(Filename, DateTime);
+}
+
+IFileHandle* FStorageServerPlatformFile::OpenWrite(const TCHAR* Filename, bool bAppend, bool bAllowRead)
+{
+	TStringBuilder<1024> StorageServerFilename;
+	if (MakeStorageServerPath(Filename, StorageServerFilename) && ServerToc.FileExists(*StorageServerFilename))
+	{
+		return nullptr;
+	}
+	return LowerLevel->OpenWrite(Filename, bAppend, bAllowRead);
+}
+
+bool FStorageServerPlatformFile::CreateDirectory(const TCHAR* Directory)
+{
+	TStringBuilder<1024> StorageServerDirectory;
+	if (MakeStorageServerPath(Directory, StorageServerDirectory) && ServerToc.DirectoryExists(*StorageServerDirectory))
+	{
+		return true;
+	}
+	return LowerLevel->CreateDirectory(Directory);
+}
+
+bool FStorageServerPlatformFile::DeleteDirectory(const TCHAR* Directory)
+{
+	TStringBuilder<1024> StorageServerDirectory;
+	if (MakeStorageServerPath(Directory, StorageServerDirectory) && ServerToc.DirectoryExists(*StorageServerDirectory))
+	{
+		return false;
+	}
+	return LowerLevel->DeleteDirectory(Directory);
+}
+
+bool FStorageServerPlatformFile::MakeStorageServerPath(const TCHAR* LocalFilenameOrDirectory, FStringBuilderBase& OutPath) const
+{
+	FStringView LocalEngineDirView(FPlatformMisc::EngineDir());
+	FStringView LocalProjectDirView(FPlatformMisc::ProjectDir());
+	FStringView LocalFilenameOrDirectoryView(LocalFilenameOrDirectory);
+	bool bValid = false;
+
+	if (LocalFilenameOrDirectoryView.StartsWith(LocalEngineDirView, ESearchCase::IgnoreCase))
+	{
+		OutPath.Append(ServerEngineDirView);
+		OutPath.Append(LocalFilenameOrDirectoryView.RightChop(LocalEngineDirView.Len()));
+		bValid = true;
+	}
+	else if (LocalFilenameOrDirectoryView.StartsWith(LocalProjectDirView, ESearchCase::IgnoreCase))
+	{
+		OutPath.Append(ServerProjectDirView);
+		OutPath.Append(LocalFilenameOrDirectoryView.RightChop(LocalProjectDirView.Len()));
+		bValid = true;
+	}
+
+	if (bValid)
+	{
+		Algo::Replace(MakeArrayView(OutPath), '\\', '/');
+		OutPath.RemoveSuffix(LocalFilenameOrDirectoryView.EndsWith('/') ? 1 : 0);
+	}
+
+	return bValid;
+}
+
+bool FStorageServerPlatformFile::MakeLocalPath(const TCHAR* ServerFilenameOrDirectory, FStringBuilderBase& OutPath) const
+{
+	FStringView ServerFilenameOrDirectoryView(ServerFilenameOrDirectory);
+	if (ServerFilenameOrDirectoryView.StartsWith(ServerEngineDirView, ESearchCase::IgnoreCase))
+	{
+		OutPath.Append(FPlatformMisc::EngineDir());
+		OutPath.Append(ServerFilenameOrDirectoryView.RightChop(ServerEngineDirView.Len()));
+		return true;
+	}
+	else if (ServerFilenameOrDirectoryView.StartsWith(ServerProjectDirView, ESearchCase::IgnoreCase))
+	{
+		OutPath.Append(FPlatformMisc::ProjectDir());
+		OutPath.Append(ServerFilenameOrDirectoryView.RightChop(ServerProjectDirView.Len()));
+		return true;
+	}
+	return false;
+}
+
+bool FStorageServerPlatformFile::SendGetFileListMessage()
+{
+	TRACE_CPUPROFILER_EVENT_SCOPE(StorageServerPlatformFileGetFileList);
+	
+	Connection->FileManifestRequest([&](FIoChunkId Id, FStringView Path)
+	{
+		ServerToc.AddFile(Path, FileIndexFromChunkId(Id));
+	});
+
+	return true;
+}
+
+FFileStatData FStorageServerPlatformFile::SendGetStatDataMessage(int32 FileIndex)
+{
+	TRACE_CPUPROFILER_EVENT_SCOPE(StorageServerPlatformFileGetStatData);
+	int64 FileSize = Connection->FileSizeRequest(FileIndex);
+	if (FileSize < 0)
+	{
+		return FFileStatData();
+	}
+
+	FDateTime CreationTime = FDateTime::Now();
+	FDateTime AccessTime = FDateTime::Now();
+	FDateTime ModificationTime = FDateTime::Now();
+
+	return FFileStatData(CreationTime, AccessTime, ModificationTime, FileSize, false, true);
+}
+
+int64 FStorageServerPlatformFile::SendReadMessage(uint8* Destination, int32 FileIndex, int64 Offset, int64 BytesToRead)
+{
+	TRACE_CPUPROFILER_EVENT_SCOPE(StorageServerPlatformFileRead);
+	int64 BytesRead = 0;
+	Connection->ReadFileRequest(FileIndex, Offset, BytesToRead, [Destination, &BytesRead](FStorageServerResponse& Response)
+	{
+		BytesRead = Response.TotalSize();
+		Response.Serialize(Destination, Response.TotalSize());
+	});
+	return BytesRead;
+}
+
+#if WITH_COTF
+void FStorageServerPlatformFile::OnCookOnTheFlyMessage(const UE::Cook::FCookOnTheFlyMessage& Message)
+{
+	switch (Message.GetHeader().MessageType)
+	{
+		case UE::Cook::ECookOnTheFlyMessage::FilesAdded:
+		{
+			UE_LOG(LogCookOnTheFly, Verbose, TEXT("Received '%s' message"), LexToString(Message.GetHeader().MessageType));
+
+			TArray<FString> Filenames;
+			TArray<FIoChunkId> ChunkIds;
+
+			{
+				TUniquePtr<FArchive> Ar = Message.ReadBody();
+				*Ar << Filenames;
+				*Ar << ChunkIds;
+			}
+
+			check(Filenames.Num() == ChunkIds.Num());
+
+			for (int32 Idx = 0, Num = Filenames.Num(); Idx < Num; ++Idx)
+			{
+				UE_LOG(LogCookOnTheFly, Verbose, TEXT("Adding file '%s'"), *Filenames[Idx]);
+				ServerToc.AddFile(Filenames[Idx], FileIndexFromChunkId(ChunkIds[Idx]));
+			}
+
+			break;
+		}
+	}
+}
+#endif
+
+class FStorageServerClientFileModule
+	: public IPlatformFileModule
+{
+public:
+
+	virtual IPlatformFile* GetPlatformFile() override
+	{
+		static TUniquePtr<IPlatformFile> AutoDestroySingleton = MakeUnique<FStorageServerPlatformFile>();
+		return AutoDestroySingleton.Get();
+	}
+};
+
+IMPLEMENT_MODULE(FStorageServerClientFileModule, StorageServerClient);
+
+#endif
diff --git a/Engine/Source/Runtime/StorageServerClient/Public/StorageServerConnection.h b/Engine/Source/Runtime/StorageServerClient/Public/StorageServerConnection.h
new file mode 100644
--- /dev/null
+++ b/Engine/Source/Runtime/StorageServerClient/Public/StorageServerConnection.h
@@ -0,0 +1,141 @@
+// Copyright Epic Games, Inc. All Rights Reserved.
+
+#pragma once
+
+#include "CoreMinimal.h"
+#include "IO/IoContainerId.h"
+#include "Misc/StringBuilder.h"
+#include "Serialization/CompactBinarySerialization.h"
+
+#if !UE_BUILD_SHIPPING
+
+class ISocketSubsystem;
+class FInternetAddr;
+class FIoChunkId;
+class FSocket;
+class FStorageServerConnection;
+
+class FStorageServerRequest
+	: public FArchive
+{
+protected:
+	friend class FStorageServerConnection;
+
+	FStorageServerRequest(FAnsiStringView Verb, FAnsiStringView Resource, FAnsiStringView Hostname);
+	FSocket* Send(FStorageServerConnection& Owner);
+	virtual void Serialize(void* V, int64 Length) override;
+
+	TAnsiStringBuilder<512> HeaderBuffer;
+	TArray<uint8, TInlineAllocator<1024>> BodyBuffer;
+};
+
+class FStorageServerResponse
+	: public FArchive
+{
+public:
+	~FStorageServerResponse()
+	{
+		if (Socket)
+		{
+			ReleaseSocket(false);
+		}
+	}
+
+	bool IsOk() const
+	{
+		return bIsOk;
+	}
+
+	const int32 GetErrorCode() const
+	{
+		return ErrorCode;
+	}
+
+	const FString& GetErrorMessage() const
+	{
+		return ErrorMessage;
+	}
+
+	int64 TotalSize() override
+	{
+		return ContentLength;
+	}
+
+	int64 Tell() override
+	{
+		return Position;
+	}
+
+	void Serialize(void* V, int64 Length) override;
+
+	FCbObject GetResponseObject()
+	{
+		FCbField Payload = LoadCompactBinary(*this);
+		return Payload.AsObject();
+	}
+
+private:
+	friend class FStorageServerConnection;
+	friend class FStorageServerChunkBatchRequest;
+
+	FStorageServerResponse(class FStorageServerConnection& Owner, FSocket& Socket);
+	void ReleaseSocket(bool bKeepAlive);
+
+	FStorageServerConnection& Owner;
+	FSocket* Socket = nullptr;
+	uint64 ContentLength = 0;
+	uint64 Position = 0;
+	int32 ErrorCode;
+	FString ErrorMessage;
+	bool bIsOk = false;
+};
+
+class FStorageServerChunkBatchRequest
+	: private FStorageServerRequest
+{
+public:
+	STORAGESERVERCLIENT_API FStorageServerChunkBatchRequest& AddChunk(const FIoChunkId& ChunkId, int64 Offset, int64 Size);
+	STORAGESERVERCLIENT_API bool Issue(TFunctionRef<void(uint32 ChunkCount, uint32* ChunkIndices, uint64* ChunkSizes, FStorageServerResponse& ChunkDataStream)> OnResponse);
+
+private:
+	friend class FStorageServerConnection;
+
+	FStorageServerChunkBatchRequest(FStorageServerConnection& Owner, FAnsiStringView Resource, FAnsiStringView Hostname);
+
+	FStorageServerConnection& Owner;
+	int32 ChunkCountOffset = 0;
+};
+
+class FStorageServerConnection
+{
+public:
+	STORAGESERVERCLIENT_API FStorageServerConnection();
+	STORAGESERVERCLIENT_API ~FStorageServerConnection();
+
+	STORAGESERVERCLIENT_API bool Initialize(const TCHAR* Host, int32 Port, const TCHAR* ProjectNameOverride = nullptr, const TCHAR* PlatformNameOverride = nullptr);
+
+	STORAGESERVERCLIENT_API void FileManifestRequest(TFunctionRef<void(FIoChunkId Id, FStringView Path)> Callback);
+	STORAGESERVERCLIENT_API int64 FileSizeRequest(int32 FileIndex);
+	STORAGESERVERCLIENT_API bool ReadFileRequest(int32 FileIndex, uint64 Offset, uint64 Size, TFunctionRef<void(FStorageServerResponse&)> OnResponse);
+	STORAGESERVERCLIENT_API int64 ChunkSizeRequest(const FIoChunkId& ChunkId);
+	STORAGESERVERCLIENT_API bool ReadChunkRequest(const FIoChunkId& ChunkId, uint64 Offset, uint64 Size, TFunctionRef<void(FStorageServerResponse&)> OnResponse);
+	STORAGESERVERCLIENT_API FStorageServerChunkBatchRequest NewChunkBatchRequest();
+
+private:
+	friend class FStorageServerRequest;
+	friend class FStorageServerResponse;
+	friend class FStorageServerChunkBatchRequest;
+
+	int32 HandshakeRequest();
+	FSocket* AcquireSocket();
+	void ReleaseSocket(FSocket* Socket, bool bKeepAlive);
+
+	ISocketSubsystem& SocketSubsystem;
+	TAnsiStringBuilder<1024> OplogPath;
+	TSharedPtr<FInternetAddr> ServerAddr;
+	TAnsiStringBuilder<1024> Hostname;
+	FCriticalSection SocketPoolCritical;
+	TArray<FSocket*> SocketPool;
+};
+
+#endif
