Don't visit inner platform file if directory exist

- prevents issues where the same directory/file can be visisted twice

#rnx

diff --git a/Engine/Source/Runtime/StorageServerClient/Private/StorageServerPlatformFile.cpp b/Engine/Source/Runtime/StorageServerClient/Private/StorageServerPlatformFile.cpp
--- a/Engine/Source/Runtime/StorageServerClient/Private/StorageServerPlatformFile.cpp
+++ b/Engine/Source/Runtime/StorageServerClient/Private/StorageServerPlatformFile.cpp
@@ -414,7 +414,7 @@ bool FStorageServerPlatformFile::IterateDirectory(const TCHAR* Directory, IPlatf
 {
 	TStringBuilder<1024> StorageServerDirectory;
 	bool bResult = false;
-	if (MakeStorageServerPath(Directory, StorageServerDirectory))
+	if (MakeStorageServerPath(Directory, StorageServerDirectory) && ServerToc.DirectoryExists(*StorageServerDirectory))
 	{
 		bResult |= ServerToc.IterateDirectory(*StorageServerDirectory, [this, &Visitor](int32 FileIndex, const TCHAR* FilenameOrDirectory)
 		{
@@ -424,7 +424,10 @@ bool FStorageServerPlatformFile::IterateDirectory(const TCHAR* Directory, IPlatf
 			return Visitor.Visit(*LocalPath, FileIndex < 0);
 		});
 	}
-	bResult |= LowerLevel->IterateDirectory(Directory, Visitor);
+	else
+	{
+		bResult |= LowerLevel->IterateDirectory(Directory, Visitor);
+	}
 	return bResult;
 }
 
@@ -432,7 +435,7 @@ bool FStorageServerPlatformFile::IterateDirectoryStat(const TCHAR* Directory, FD
 {
 	TStringBuilder<1024> StorageServerDirectory;
 	bool bResult = false;
-	if (MakeStorageServerPath(Directory, StorageServerDirectory))
+	if (MakeStorageServerPath(Directory, StorageServerDirectory) && ServerToc.DirectoryExists(*StorageServerDirectory))
 	{
 		bResult |= ServerToc.IterateDirectory(*StorageServerDirectory, [this, &Visitor](int32 FileIndex, const TCHAR* ServerFilenameOrDirectory)
 		{
@@ -458,7 +461,10 @@ bool FStorageServerPlatformFile::IterateDirectoryStat(const TCHAR* Directory, FD
 			return Visitor.Visit(*LocalPath, FileStatData);
 		});
 	}
-	bResult |= LowerLevel->IterateDirectoryStat(Directory, Visitor);
+	else
+	{
+		bResult |= LowerLevel->IterateDirectoryStat(Directory, Visitor);
+	}
 	return bResult;
 }
 
