Simplify command line arguments when running COTF with Zen loader/storage server

- changed Cooker arguments from -CookOnTheFly -IoStore to -CookOnTheFly -ZenStore
- changed Zen host argument from -StorageServerHost to -ZenStoreHost
- updated UAT to support running CBTB with -ZenStore
- updated StorageServerConnection to handle multiple hosts from command line

Cooker:
  CBTB: -ZenStore
  COTF: -ZenStore -CookOnTheFly

Game:
  CBTB: -ZenStoreHost=<ip> or <ip1,ip2>
  COTF: -ZenStoreHost=<ip> or <ip1,ip2> -CookOnTheFly

#rnx

diff --git a/Engine/Source/Developer/LauncherServices/Private/Launcher/LauncherWorker.cpp b/Engine/Source/Developer/LauncherServices/Private/Launcher/LauncherWorker.cpp
--- a/Engine/Source/Developer/LauncherServices/Private/Launcher/LauncherWorker.cpp
+++ b/Engine/Source/Developer/LauncherServices/Private/Launcher/LauncherWorker.cpp
@@ -695,9 +695,9 @@ FString FLauncherWorker::CreateUATCommand( const ILauncherProfileRef& InProfile,
 		{
 			UATCommand += TEXT(" -cookonthefly");
 			
-			if (InProfile->IsUsingIoStore())
+			if (InProfile->IsUsingZenStore())
 			{
-				UATCommand += TEXT(" -iostore");
+				UATCommand += TEXT(" -zenstore");
 			}
 
 			//if UAT doesn't stick around as long as the process we are going to run, then we can't kill the COTF server when UAT goes down because the program
@@ -724,9 +724,9 @@ FString FLauncherWorker::CreateUATCommand( const ILauncherProfileRef& InProfile,
 	case ELauncherProfileCookModes::OnTheFlyInEditor:
 		UATCommand += MapList;
 		UATCommand += " -skipcook -cookonthefly -CookInEditor";
-		if (InProfile->IsUsingIoStore())
+		if (InProfile->IsUsingZenStore())
 		{
-			UATCommand += TEXT(" -iostore");
+			UATCommand += TEXT(" -zenstore");
 		}
 		break;
 	case ELauncherProfileCookModes::ByTheBookInEditor:
diff --git a/Engine/Source/Developer/LauncherServices/Private/Profiles/LauncherProfile.h b/Engine/Source/Developer/LauncherServices/Private/Profiles/LauncherProfile.h
--- a/Engine/Source/Developer/LauncherServices/Private/Profiles/LauncherProfile.h
+++ b/Engine/Source/Developer/LauncherServices/Private/Profiles/LauncherProfile.h
@@ -2438,6 +2438,16 @@ public:
 		return bUseIoStore;
 	}
 
+	virtual void SetUseZenStore(bool bInUseZenStore) override
+	{
+		bUseZenStore = bInUseZenStore;
+	}
+
+	virtual bool IsUsingZenStore() const override
+	{
+		return bUseZenStore;
+	}
+
 	virtual void SetShouldUpdateDeviceFlash(bool bInShouldUpdateFlash) override
 	{
 		bShouldUpdateFlash = bInShouldUpdateFlash;
@@ -2914,6 +2924,9 @@ private:
 
 	// Use I/O store.
 	bool bUseIoStore;
+	
+	// Use Zen storage server
+	bool bUseZenStore;
 
 	// Make binary config.
 	bool bMakeBinaryConfig;
diff --git a/Engine/Source/Developer/LauncherServices/Public/ILauncherProfile.h b/Engine/Source/Developer/LauncherServices/Public/ILauncherProfile.h
--- a/Engine/Source/Developer/LauncherServices/Public/ILauncherProfile.h
+++ b/Engine/Source/Developer/LauncherServices/Public/ILauncherProfile.h
@@ -1303,6 +1303,19 @@ public:
 	 */
 	virtual bool IsUsingIoStore() const = 0;
 
+	/**
+	 * Sets whether to use the Zen storage server
+	 * @param bUseZenStore Whether to use the Zen storage server
+	 */
+	virtual void SetUseZenStore(bool bUseZenStore) = 0;
+
+	/**
+	 * Using Zen storage server or not.
+	 *
+	 * @return true if using Zen storage server
+	 */
+	virtual bool IsUsingZenStore() const = 0;
+
 	/**
 	 * Sets whether to make a binary config file during packaging
 	 * @param bMakeBinaryConfig Whether to make a binary config file during staging
diff --git a/Engine/Source/Editor/UnrealEd/Classes/Commandlets/CookCommandlet.h b/Engine/Source/Editor/UnrealEd/Classes/Commandlets/CookCommandlet.h
--- a/Engine/Source/Editor/UnrealEd/Classes/Commandlets/CookCommandlet.h
+++ b/Engine/Source/Editor/UnrealEd/Classes/Commandlets/CookCommandlet.h
@@ -55,8 +55,6 @@ class UCookCommandlet
 	bool bPartialGC;
 	/** Do not cook any shaders.  Shader maps will be empty */
 	bool bNoShaderCooking;
-	/** Whether to cook for I/O store */
-	bool bUseIoStore;
 	/** All commandline tokens */
 	TArray<FString> Tokens;
 	/** All commandline switches */
diff --git a/Engine/Source/Editor/UnrealEd/Classes/CookOnTheSide/CookOnTheFlyServer.h b/Engine/Source/Editor/UnrealEd/Classes/CookOnTheSide/CookOnTheFlyServer.h
--- a/Engine/Source/Editor/UnrealEd/Classes/CookOnTheSide/CookOnTheFlyServer.h
+++ b/Engine/Source/Editor/UnrealEd/Classes/CookOnTheSide/CookOnTheFlyServer.h
@@ -146,8 +146,8 @@ public:
 	{
 		/** Wether the network file server or the I/O store connection server should bind to any port */
 		bool bBindAnyPort = false;
-		/** Wether to cook for I/O store */
-		bool bIoStore = false;
+		/** Whether to save the cooked output to the Zen storage server. */
+		bool bZenStore = false;
 		/** Target platforms */
 		TArray<ITargetPlatform*> TargetPlatforms;
 	};
diff --git a/Engine/Source/Editor/UnrealEd/Classes/Settings/LevelEditorPlaySettings.h b/Engine/Source/Editor/UnrealEd/Classes/Settings/LevelEditorPlaySettings.h
--- a/Engine/Source/Editor/UnrealEd/Classes/Settings/LevelEditorPlaySettings.h
+++ b/Engine/Source/Editor/UnrealEd/Classes/Settings/LevelEditorPlaySettings.h
@@ -274,9 +274,9 @@ public:
 	UPROPERTY(config, EditAnywhere, Category = PlayInEditor, meta = (DisplayName="Stream Sub-Levels during Play in Editor", ToolTip="Prefer to stream sub-levels from the disk instead of duplicating editor sub-levels"))
 	uint32 bPreferToStreamLevelsInPIE:1;
 
-	/** Whether to use Zen Loader/Storage server when running cook-on-the-fly. */
+	/** Whether to use Zen loader/storage server when running cook-on-the-fly. */
 	UPROPERTY(config, EditAnywhere, Category=PlayInEditor)
-	bool EnableIoStoreCookOnTheFlyMode;
+	bool EnableZenCookOnTheFlyMode;
 
 public:
 	/** The width of the new view port window in pixels (0 = use the desktop's screen resolution). */
diff --git a/Engine/Source/Editor/UnrealEd/Private/Commandlets/CookCommandlet.cpp b/Engine/Source/Editor/UnrealEd/Private/Commandlets/CookCommandlet.cpp
--- a/Engine/Source/Editor/UnrealEd/Private/Commandlets/CookCommandlet.cpp
+++ b/Engine/Source/Editor/UnrealEd/Private/Commandlets/CookCommandlet.cpp
@@ -424,7 +424,7 @@ bool UCookCommandlet::CookOnTheFly( FGuid InstanceId, int32 Timeout, bool bForce
 
 	UCookOnTheFlyServer::FCookOnTheFlyOptions CookOnTheFlyStartupOptions;
 	CookOnTheFlyStartupOptions.bBindAnyPort = InstanceId.IsValid();
-	CookOnTheFlyStartupOptions.bIoStore = bUseIoStore;
+	CookOnTheFlyStartupOptions.bZenStore = Switches.Contains(TEXT("ZenStore"));
 	CookOnTheFlyStartupOptions.TargetPlatforms = TargetPlatforms;
 
 	if (CookOnTheFlyServer->StartCookOnTheFly(MoveTemp(CookOnTheFlyStartupOptions)) == false)
@@ -635,7 +635,6 @@ int32 UCookCommandlet::Main(const FString& CmdLineParams)
 	bPartialGC = Switches.Contains(TEXT("Partialgc"));
 	ShowErrorCount = !Switches.Contains(TEXT("DIFFONLY"));
 	ShowProgress = !Switches.Contains(TEXT("DIFFONLY"));
-	bUseIoStore = Switches.Contains(TEXT("IoStore"));
 	bNoShaderCooking = bCookOnTheFly; // Do not cook any shaders into the shader maps. Always true if we are running w/ cook on the fly
 
 	COOK_STAT(DetailedCookStats::CookProject = FApp::GetProjectName());
diff --git a/Engine/Source/Editor/UnrealEd/Private/CookOnTheFlyServer.cpp b/Engine/Source/Editor/UnrealEd/Private/CookOnTheFlyServer.cpp
--- a/Engine/Source/Editor/UnrealEd/Private/CookOnTheFlyServer.cpp
+++ b/Engine/Source/Editor/UnrealEd/Private/CookOnTheFlyServer.cpp
@@ -595,7 +595,18 @@ public:
 
 	virtual IPackageStoreWriter* GetPackageStoreWriter(const FName& PlatformName) override
 	{
-		return Cooker.GetPackageStoreWriter(PlatformName);
+		if (Cooker.CookOnTheFlyOptions.bZenStore)
+		{
+			for (FSavePackageContext* Context : Cooker.SavePackageContexts)
+			{
+				if (FName(*Context->TargetPlatform->PlatformName()) == PlatformName)
+				{
+					return Context->PackageStoreWriter;
+				}
+			}
+		}
+
+		return nullptr;
 	}
 
 	DECLARE_DERIVED_EVENT(UCookOnTheFlyServer::FCookOnTheFlyServerInterface, UE::Cook::ICookOnTheFlyServer::FFlushEvent, FFlushEvent);
@@ -620,7 +631,7 @@ public:
 
 		UE_LOG(LogCook, Log, TEXT("Flushing..."));
 
-		if (Cooker.CookOnTheFlyOptions.bIoStore)
+		if (Cooker.CookOnTheFlyOptions.bZenStore)
 		{
 			for (FSavePackageContext* Context : Cooker.SavePackageContexts)
 			{
@@ -804,16 +815,17 @@ bool UCookOnTheFlyServer::StartCookOnTheFly(FCookOnTheFlyOptions InCookOnTheFlyO
 	}
 	}
 
-	if (CookOnTheFlyOptions.bIoStore)
+	UE_LOG(LogCook, Display, TEXT("Starting '%s' cook-on-the-fly server"),
+		CookOnTheFlyOptions.bZenStore ? TEXT("Zen") : TEXT("Network File"));
+
+	if (CookOnTheFlyOptions.bZenStore)
 	{
-		UE_LOG(LogCook, Display, TEXT("Initializing I/O store cook on the fly server"));
 		UE::Cook::FIoStoreCookOnTheFlyServerOptions ServerOptions;
 		ServerOptions.Port = -1; // Use default
 		CookOnTheFlyRequestManager = UE::Cook::MakeIoStoreCookOnTheFlyRequestManager(*CookOnTheFlyServerInterface, MoveTemp(ServerOptions));
 	}
 	else
 	{
-		UE_LOG(LogCook, Display, TEXT("Initializing Network File cook on the fly server"));
 		FNetworkFileServerOptions ServerOptions;
 		ServerOptions.Protocol = NFSP_Tcp;
 		ServerOptions.Port = CookOnTheFlyOptions.bBindAnyPort ? 0 : -1;
@@ -1285,12 +1297,12 @@ bool UCookOnTheFlyServer::IsUsingShaderCodeLibrary() const
 
 bool UCookOnTheFlyServer::IsUsingIoStore() const
 {
-	return (IsCookByTheBookMode() && CookByTheBookOptions->bIoStore) || (IsCookOnTheFlyMode() && CookOnTheFlyOptions.bIoStore);
+	return (IsCookByTheBookMode() && CookByTheBookOptions->bIoStore) || (IsCookOnTheFlyMode() && CookOnTheFlyOptions.bZenStore);
 }
 
 bool UCookOnTheFlyServer::IsUsingZenStore() const
 {
-	return (IsCookByTheBookMode() && CookByTheBookOptions->bZenStore) || (IsCookOnTheFlyMode() && CookOnTheFlyOptions.bIoStore);
+	return (IsCookByTheBookMode() && CookByTheBookOptions->bZenStore) || (IsCookOnTheFlyMode() && CookOnTheFlyOptions.bZenStore);
 }
 
 bool UCookOnTheFlyServer::IsCookOnTheFlyMode() const
diff --git a/Engine/Source/Editor/UnrealEd/Private/PlayLevelViaLauncher.cpp b/Engine/Source/Editor/UnrealEd/Private/PlayLevelViaLauncher.cpp
--- a/Engine/Source/Editor/UnrealEd/Private/PlayLevelViaLauncher.cpp
+++ b/Engine/Source/Editor/UnrealEd/Private/PlayLevelViaLauncher.cpp
@@ -288,7 +288,7 @@ void UEditorEngine::StartPlayUsingLauncherSession(FRequestPlaySessionParams& InR
 	if (LauncherProfile->GetCookMode() == ELauncherProfileCookModes::OnTheFlyInEditor || LauncherProfile->GetCookMode() == ELauncherProfileCookModes::OnTheFly)
 	{
 		LauncherProfile->SetDeploymentMode(ELauncherProfileDeploymentModes::FileServer);
-		LauncherProfile->SetUseIoStore(EditorPlaySettings->EnableIoStoreCookOnTheFlyMode);
+		LauncherProfile->SetUseZenStore(EditorPlaySettings->EnableZenCookOnTheFlyMode);
 	}
 
 	switch(EditorPlaySettings->PackFilesForLaunch)
diff --git a/Engine/Source/Programs/AutomationTool/AutomationUtils/ProjectParams.cs b/Engine/Source/Programs/AutomationTool/AutomationUtils/ProjectParams.cs
--- a/Engine/Source/Programs/AutomationTool/AutomationUtils/ProjectParams.cs
+++ b/Engine/Source/Programs/AutomationTool/AutomationUtils/ProjectParams.cs
@@ -287,6 +287,7 @@ namespace AutomationTool
 			this.IgnorePaksFromDifferentCookSource = InParams.IgnorePaksFromDifferentCookSource;
 			this.IoStore = InParams.IoStore;
 			this.Cook4IoStore = InParams.Cook4IoStore;
+			this.ZenStore = InParams.ZenStore;
 			this.GenerateOptimizationData = InParams.GenerateOptimizationData;
 			this.SignPak = InParams.SignPak;
 			this.SignedPak = InParams.SignedPak;
@@ -372,7 +373,6 @@ namespace AutomationTool
 			this.bTreatNonShippingBinariesAsDebugFiles = InParams.bTreatNonShippingBinariesAsDebugFiles;
 			this.bUseExtraFlavor = InParams.bUseExtraFlavor;
 			this.AdditionalPackageOptions = InParams.AdditionalPackageOptions;
-			this.StorageServerHost = InParams.StorageServerHost;
 		}
 
 		/// <summary>
@@ -475,6 +475,7 @@ namespace AutomationTool
 			bool? IgnorePaksFromDifferentCookSource = null,
 			bool? IoStore = null,
 			bool? Cook4IoStore = null,
+			bool? ZenStore = null,
 			bool? SkipIoStore = null,
 			bool? GenerateOptimizationData = null,
 			bool? Prereqs = null,
@@ -524,8 +525,7 @@ namespace AutomationTool
 		    bool AutomaticSigning = false,
 			ParamList<string> InMapsToRebuildLightMaps = null,
             ParamList<string> InMapsToRebuildHLOD = null,
-            ParamList<string> TitleID = null,
-			string StorageServerhost = null
+            ParamList<string> TitleID = null
 			)
 		{
 			//
@@ -669,6 +669,11 @@ namespace AutomationTool
 				this.IoStore = true;
 				this.AdditionalCookerOptions += " -IoStore";
 			}
+			this.ZenStore = GetParamValueIfNotSpecified(Command, ZenStore, this.ZenStore, "zenstore");
+			if (this.ZenStore && this.Cook && !this.SkipCook)
+			{
+				this.AdditionalCookerOptions += " -ZenStore";
+			}
 			this.GenerateOptimizationData = GetParamValueIfNotSpecified(Command, GenerateOptimizationData, this.GenerateOptimizationData, "makebinaryconfig");
 			
 			this.SkipPak = GetParamValueIfNotSpecified(Command, SkipPak, this.SkipPak, "skippak");
@@ -877,7 +882,6 @@ namespace AutomationTool
 			this.SpecifiedArchitecture = ParseParamValueIfNotSpecified(Command, SpecifiedArchitecture, "specifiedarchitecture", String.Empty);
 			this.UbtArgs = ParseParamValueIfNotSpecified(Command, UbtArgs, "ubtargs", String.Empty);
 			this.AdditionalPackageOptions = ParseParamValueIfNotSpecified(Command, AdditionalPackageOptions, "AdditionalPackageOptions", String.Empty);
-			this.StorageServerHost = ParseParamValueIfNotSpecified(Command, StorageServerHost, "StorageServerHost", "localhost");
 
 			if (ClientConfigsToBuild == null)
 			{
@@ -1247,6 +1251,12 @@ namespace AutomationTool
 		/// </summary>
 		[Help("cook4iostore", "generate I/O store container file(s)")]
 		public bool Cook4IoStore { private set; get; }
+		
+		/// <summary>
+		/// Shared: True if the cooker should store the cooked output to the Zen storage server
+		/// </summary>
+		[Help("zenstore", "save cooked output data to the Zen storage server")]
+		public bool ZenStore { private set; get; }
 
 		/// <summary>
 		/// Shared: True if optimization data is generated during staging that can improve loadtimes
@@ -2074,9 +2084,6 @@ namespace AutomationTool
         [Help("IgnoreLightMapErrors", "Whether Light Map errors should be treated as critical")]
 		public bool IgnoreLightMapErrors { get; set; }
 
-		[Help("StorageServerHost", "Storage server host address")]
-		public string StorageServerHost { get; set; }
-
 		private List<SingleTargetProperties> DetectedTargets;
 		private Dictionary<UnrealTargetPlatform, ConfigHierarchy> LoadedEngineConfigs;
 		private Dictionary<UnrealTargetPlatform, ConfigHierarchy> LoadedGameConfigs;
@@ -2957,6 +2964,7 @@ namespace AutomationTool
 				CommandUtils.LogLog("IoStore={0}", IoStore);
 				CommandUtils.LogLog("SkipIoStore={0}", SkipIoStore);
 				CommandUtils.LogLog("Cook4IoStore={0}", Cook4IoStore);
+				CommandUtils.LogLog("ZenStore={0}", ZenStore);
 				CommandUtils.LogLog("SkipEncryption={0}", SkipEncryption);
 				CommandUtils.LogLog("GenerateOptimizationData={0}", GenerateOptimizationData);
 				CommandUtils.LogLog("SkipPackage={0}", SkipPackage);
diff --git a/Engine/Source/Programs/AutomationTool/Scripts/CookCommand.Automation.cs b/Engine/Source/Programs/AutomationTool/Scripts/CookCommand.Automation.cs
--- a/Engine/Source/Programs/AutomationTool/Scripts/CookCommand.Automation.cs
+++ b/Engine/Source/Programs/AutomationTool/Scripts/CookCommand.Automation.cs
@@ -65,9 +65,9 @@ public partial class Project : CommandUtils
 					COTFCommandLine += " -ddc=" + Params.DDCGraph;
 				}
 
-				if (Params.IoStore)
+				if (Params.ZenStore)
 				{
-					COTFCommandLine += " -iostore";
+					COTFCommandLine += " -zenstore";
 				}
 
 				var ServerLogFile = CombinePaths(LogFolderOutsideOfSandbox, "Server.log");
diff --git a/Engine/Source/Programs/AutomationTool/Scripts/CopyBuildToStagingDirectory.Automation.cs b/Engine/Source/Programs/AutomationTool/Scripts/CopyBuildToStagingDirectory.Automation.cs
--- a/Engine/Source/Programs/AutomationTool/Scripts/CopyBuildToStagingDirectory.Automation.cs
+++ b/Engine/Source/Programs/AutomationTool/Scripts/CopyBuildToStagingDirectory.Automation.cs
@@ -3896,7 +3896,7 @@ public partial class Project : CommandUtils
 			string FileHostParams = " ";
 			if (Params.CookOnTheFly || Params.FileServer)
 			{
-				FileHostParams += Params.IoStore ? "-cookontheflyhost=" : "-filehostip=";
+				FileHostParams += Params.ZenStore ? "-cookonthefly -zenstorehost=" : "-filehostip=";
 				// add localhost first for platforms using redirection
 				const string LocalHost = "127.0.0.1";
 				if (!IsNullOrEmpty(Params.Port))
@@ -4013,14 +4013,7 @@ public partial class Project : CommandUtils
 					}
 				}
 
-				if (Params.IoStore)
-				{
-					FileHostParams += string.Format(" -storageserverhost={0} ", Params.StorageServerHost);
-				}
-				else
-				{
-					FileHostParams += " ";
-				}
+				FileHostParams += " ";
 			}
 
 			String ProjectFile = String.Format("{0} ", SC.ProjectArgForCommandLines);
diff --git a/Engine/Source/Programs/AutomationTool/Scripts/RunProjectCommand.Automation.cs b/Engine/Source/Programs/AutomationTool/Scripts/RunProjectCommand.Automation.cs
--- a/Engine/Source/Programs/AutomationTool/Scripts/RunProjectCommand.Automation.cs
+++ b/Engine/Source/Programs/AutomationTool/Scripts/RunProjectCommand.Automation.cs
@@ -629,9 +629,9 @@ public partial class Project : CommandUtils
 
 			if (Params.CookOnTheFly || Params.FileServer)
 			{
-				if (Params.IoStore)
+				if (Params.ZenStore)
 				{
-					TempCmdLine += "-cookontheflyhost=";
+					TempCmdLine += "-cookonthefly -zenstorehost=";
 				}
 				else
 				{
@@ -753,14 +753,7 @@ public partial class Project : CommandUtils
 					}
 				}
 
-				if (Params.IoStore)
-				{
-					TempCmdLine += string.Format(" -storageserverhost={0} ", Params.StorageServerHost);
-				}
-				else
-				{
-					TempCmdLine += " ";
-				}
+				TempCmdLine += " ";
 
 				if (Params.CookOnTheFlyStreaming)
 				{
diff --git a/Engine/Source/Runtime/CookOnTheFly/Private/CookOnTheFly.cpp b/Engine/Source/Runtime/CookOnTheFly/Private/CookOnTheFly.cpp
--- a/Engine/Source/Runtime/CookOnTheFly/Private/CookOnTheFly.cpp
+++ b/Engine/Source/Runtime/CookOnTheFly/Private/CookOnTheFly.cpp
@@ -65,15 +65,19 @@ bool GetCookOnTheFlyHost(UE::Cook::FCookOnTheFlyHostOptions& OutHostOptions)
 {
 	bool bOptionsOk = false;
 
-	FString Host;
-	if (FParse::Value(FCommandLine::Get(), TEXT("-CookOnTheFlyHost="), Host))
+	if (FParse::Param(FCommandLine::Get(), TEXT("CookOnTheFly")))
 	{
-		if (!Host.ParseIntoArray(OutHostOptions.Hosts, TEXT("+"), true))
+		// Cook-on-the-fly expects the same host as the Zen storage server
+		FString Host;
+		if (FParse::Value(FCommandLine::Get(), TEXT("-ZenStoreHost="), Host))
 		{
-			OutHostOptions.Hosts.Add(Host);
-		}
+			if (!Host.ParseIntoArray(OutHostOptions.Hosts, TEXT("+"), true))
+			{
+				OutHostOptions.Hosts.Add(Host);
+			}
 
-		bOptionsOk = true;
+			bOptionsOk = true;
+		}
 	}
 
 	double ServerWaitTimeInSeconds;
diff --git a/Engine/Source/Runtime/Core/Private/Misc/CoreGlobals.cpp b/Engine/Source/Runtime/Core/Private/Misc/CoreGlobals.cpp
--- a/Engine/Source/Runtime/Core/Private/Misc/CoreGlobals.cpp
+++ b/Engine/Source/Runtime/Core/Private/Misc/CoreGlobals.cpp
@@ -720,12 +720,16 @@ bool IsRunningCookOnTheFly()
 	static struct FCookOnTheFlyCommandline
 	{
 		bool bParsed;
-		FCookOnTheFlyCommandline()
+		FCookOnTheFlyCommandline(const TCHAR* CmdLine)
 		{
 			FString Host;
-			bParsed = FParse::Value(FCommandLine::Get(), TEXT("-CookOnTheFlyHost="), Host) || FParse::Value(FCommandLine::Get(), TEXT("-FileHostIP="), Host);
+			bParsed = FParse::Param(CmdLine, TEXT("CookOnTheFly")) && FParse::Value(CmdLine, TEXT("-ZenStoreHost="), Host);
+			if (!bParsed)	
+			{
+				bParsed = FParse::Value(CmdLine, TEXT("-FileHostIP="), Host);
+			}
 		}
-	} CookOnTheFlyCommandline;
+	} CookOnTheFlyCommandline(FCommandLine::Get());
 	
 	return CookOnTheFlyCommandline.bParsed;
 #else
diff --git a/Engine/Source/Runtime/StorageServerClient/Private/StorageServerConnection.cpp b/Engine/Source/Runtime/StorageServerClient/Private/StorageServerConnection.cpp
--- a/Engine/Source/Runtime/StorageServerClient/Private/StorageServerConnection.cpp
+++ b/Engine/Source/Runtime/StorageServerClient/Private/StorageServerConnection.cpp
@@ -16,6 +16,33 @@
 
 DEFINE_LOG_CATEGORY_STATIC(LogStorageServerConnection, Log, All);
 
+static TArray<TSharedPtr<FInternetAddr>> GetAddressFromString(ISocketSubsystem& SocketSubsystem, TArrayView<const FString> HostAddresses, const int32 Port)
+{
+	TArray<TSharedPtr<FInternetAddr>> InterntAddresses;
+
+	for (const FString& HostAddr : HostAddresses)
+	{
+		TSharedPtr<FInternetAddr> Addr = SocketSubsystem.GetAddressFromString(HostAddr);
+
+		if (!Addr.IsValid() || !Addr->IsValid())
+		{
+			FAddressInfoResult GAIRequest = SocketSubsystem.GetAddressInfo(*HostAddr, nullptr, EAddressInfoFlags::Default, NAME_None);
+			if (GAIRequest.ReturnCode == SE_NO_ERROR && GAIRequest.Results.Num() > 0)
+			{
+				Addr = GAIRequest.Results[0].Address;
+			}
+		}
+
+		if (Addr.IsValid() && Addr->IsValid())
+		{
+			Addr->SetPort(Port);
+			InterntAddresses.Emplace(MoveTemp(Addr));
+		}
+	}
+
+	return InterntAddresses;
+}
+
 FStorageServerRequest::FStorageServerRequest(FAnsiStringView Verb, FAnsiStringView Resource, FAnsiStringView Hostname)
 {
 	SetIsSaving(true);
@@ -257,20 +284,16 @@ FStorageServerConnection::~FStorageServerConnection()
 	}
 }
 
-bool FStorageServerConnection::Initialize(const TCHAR* InHost, int32 InPort, const TCHAR* InProjectNameOverride, const TCHAR* InPlatformNameOverride)
+bool FStorageServerConnection::Initialize(TArrayView<const FString> InHostAddresses, int32 InPort, const TCHAR* InProjectNameOverride, const TCHAR* InPlatformNameOverride)
 {
-	FAddressInfoResult GAIRequest = SocketSubsystem.GetAddressInfo(InHost, nullptr, EAddressInfoFlags::Default, NAME_None);
-	if (GAIRequest.ReturnCode == SE_NO_ERROR && GAIRequest.Results.Num() > 0)
-	{
-		ServerAddr = GAIRequest.Results[0].Address;
-	}
-	if (!ServerAddr.IsValid() || !ServerAddr->IsValid())
+	TArray<TSharedPtr<FInternetAddr>> HostAddresses = GetAddressFromString(SocketSubsystem, InHostAddresses, InPort);
+
+	if (!HostAddresses.Num())
 	{
-		UE_LOG(LogStorageServerConnection, Fatal, TEXT("Failed to resolve storage server host %s."), InHost);
+		UE_LOG(LogStorageServerConnection, Fatal, TEXT("No valid Zen store host address specified"));
 		return false;
 	}
 	
-	ServerAddr->SetPort(InPort);
 	OplogPath.Append("/prj/");
 	if (InProjectNameOverride)
 	{
@@ -290,47 +313,56 @@ bool FStorageServerConnection::Initialize(const TCHAR* InHost, int32 InPort, con
 		OplogPath.Append(FPlatformProperties::PlatformName());
 	}
 
-	Hostname.Append(TCHAR_TO_ANSI(*ServerAddr->ToString(false)));
-
-	int32 ServerVersion = HandshakeRequest();
+	const int32 ServerVersion = HandshakeRequest(HostAddresses);
 	if (ServerVersion != 1)
 	{
 		return false;
 	}
 
-	UE_LOG(LogStorageServerConnection, Display, TEXT("Connected to storage server host %s."), *ServerAddr->ToString(true));
+	UE_LOG(LogStorageServerConnection, Display, TEXT("Connected to Zen storage server at '%s'"), *ServerAddr->ToString(true));
 
 	return true;
 }
 
-int32 FStorageServerConnection::HandshakeRequest()
+int32 FStorageServerConnection::HandshakeRequest(TArrayView<const TSharedPtr<FInternetAddr>> HostAddresses)
 {
 	TAnsiStringBuilder<256> ResourceBuilder;
 	ResourceBuilder.Append(OplogPath);
-	FStorageServerRequest Request("GET", *ResourceBuilder, Hostname);
-	FSocket* Socket = Request.Send(*this);
-	if (!Socket)
-	{
-		UE_LOG(LogStorageServerConnection, Fatal, TEXT("Failed to send handshake request to storage server at %s."), *ServerAddr->ToString(true));
-		return -1;
-	}
-	FStorageServerResponse Response(*this, *Socket);
 
-	if (Response.IsOk())
+	for (const TSharedPtr<FInternetAddr>& Addr : HostAddresses)
 	{
-		FCbObject ResponseObj = Response.GetResponseObject();
+		Hostname.Reset();
+		Hostname.Append(TCHAR_TO_ANSI(*Addr->ToString(false)));
+		ServerAddr = Addr;
+		
+		UE_LOG(LogStorageServerConnection, Display, TEXT("Trying to handshake with Zen at '%s'"), *Addr->ToString(true));
 
-		// we currently don't have any concept of protocol versioning, if
-		// we succeed in communicating with the endpoint we're good since
-		// any breaking API change would need to be done in a backward
-		// compatible manner
+		FStorageServerRequest Request("GET", *ResourceBuilder, Hostname);
+		if (FSocket* Socket = Request.Send(*this))
+		{
+			FStorageServerResponse Response(*this, *Socket);
 
-		return 1;
-	}
-	else
-	{
-		UE_LOG(LogStorageServerConnection, Fatal, TEXT("Failed to handshake with storage server at %s. '%s'"), *ServerAddr->ToString(true), *Response.GetErrorMessage());
+			if (Response.IsOk())
+			{
+				FCbObject ResponseObj = Response.GetResponseObject();
+
+				// we currently don't have any concept of protocol versioning, if
+				// we succeed in communicating with the endpoint we're good since
+				// any breaking API change would need to be done in a backward
+				// compatible manner
+
+				return 1;
+			}
+			else
+			{
+				UE_LOG(LogStorageServerConnection, Fatal, TEXT("Failed to handshake with Zen at %s. '%s'"), *ServerAddr->ToString(true), *Response.GetErrorMessage());
+			}
+		}
 	}
+
+	Hostname.Reset();
+	ServerAddr.Reset();
+	
 	return -1;
 }
 
@@ -564,6 +596,11 @@ FSocket* FStorageServerConnection::AcquireSocket()
 	return nullptr;
 }
 
+FString FStorageServerConnection::GetHostAddr() const
+{
+	return ServerAddr.IsValid() ? ServerAddr->ToString(false) : FString();
+}
+
 void FStorageServerConnection::ReleaseSocket(FSocket* Socket, bool bKeepAlive)
 {
 	if (bKeepAlive)
diff --git a/Engine/Source/Runtime/StorageServerClient/Private/StorageServerPlatformFile.cpp b/Engine/Source/Runtime/StorageServerClient/Private/StorageServerPlatformFile.cpp
--- a/Engine/Source/Runtime/StorageServerClient/Private/StorageServerPlatformFile.cpp
+++ b/Engine/Source/Runtime/StorageServerClient/Private/StorageServerPlatformFile.cpp
@@ -3,6 +3,7 @@
 #include "StorageServerPlatformFile.h"
 #include "StorageServerIoDispatcherBackend.h"
 #include "HAL/IPlatformFileModule.h"
+#include "Misc/CommandLine.h"
 #include "Misc/Paths.h"
 #include "Misc/ScopeRWLock.h"
 #include "StorageServerConnection.h"
@@ -239,25 +240,32 @@ FStorageServerPlatformFile::~FStorageServerPlatformFile()
 
 bool FStorageServerPlatformFile::ShouldBeUsed(IPlatformFile* Inner, const TCHAR* CmdLine) const
 {
-	FString StorageServerHost;
-	return FParse::Value(CmdLine, TEXT("-StorageServerHost="), StorageServerHost);
+	FString Host;
+	if (FParse::Value(FCommandLine::Get(), TEXT("-ZenStoreHost="), Host))
+	{
+		if (!Host.ParseIntoArray(HostAddrs, TEXT("+"), true))
+		{
+			HostAddrs.Add(Host);
+		}
+	}
+
+	return HostAddrs.Num() > 0;
 }
 
 bool FStorageServerPlatformFile::Initialize(IPlatformFile* Inner, const TCHAR* CmdLine)
 {
 	LowerLevel = Inner;
 	bool bResult = false;
-	FString StorageServerHost;
-	if (FParse::Value(CmdLine, TEXT("-StorageServerHost="), StorageServerHost))
+	if (HostAddrs.Num() > 0)
 	{
 		Connection.Reset(new FStorageServerConnection());
 		FString StorageServerProject;
-		FParse::Value(CmdLine, TEXT("-StorageServerProject="), StorageServerProject);
+		FParse::Value(CmdLine, TEXT("-ZenStoreProject="), StorageServerProject);
 		const TCHAR* ProjectOverride = StorageServerProject.IsEmpty() ? nullptr : *StorageServerProject;
 		FString StorageServerPlatform;
-		FParse::Value(CmdLine, TEXT("-StorageServerPlatform="), StorageServerPlatform);
+		FParse::Value(CmdLine, TEXT("-ZenStorePlatform="), StorageServerPlatform);
 		const TCHAR* PlatformOverride = StorageServerPlatform.IsEmpty() ? nullptr : *StorageServerPlatform;
-		if (Connection->Initialize(*StorageServerHost, 1337, ProjectOverride, PlatformOverride))
+		if (Connection->Initialize(HostAddrs, 1337, ProjectOverride, PlatformOverride))
 		{
 			if (SendGetFileListMessage())
 			{
@@ -278,17 +286,17 @@ bool FStorageServerPlatformFile::Initialize(IPlatformFile* Inner, const TCHAR* C
 				}
 				else
 				{
-					UE_LOG(LogStorageServerPlatformFile, Fatal, TEXT("Failed to initialize IoDispatcher"), *StorageServerHost);
+					UE_LOG(LogStorageServerPlatformFile, Fatal, TEXT("Failed to initialize IoDispatcher with Zen host '%s'"), *Connection->GetHostAddr());
 				}
 			}
 			else
 			{
-				UE_LOG(LogStorageServerPlatformFile, Fatal, TEXT("Failed to get file list"), *StorageServerHost);
+				UE_LOG(LogStorageServerPlatformFile, Fatal, TEXT("Failed to get file list from Zen at '%s'"), *Connection->GetHostAddr());
 			}
 		}
 		else
 		{
-			UE_LOG(LogStorageServerPlatformFile, Fatal, TEXT("Failed to initialize connection to %s"), *StorageServerHost);
+			UE_LOG(LogStorageServerPlatformFile, Fatal, TEXT("Failed to initialize connection"));
 		}
 	}
 	if (!bResult)
diff --git a/Engine/Source/Runtime/StorageServerClient/Private/StorageServerPlatformFile.h b/Engine/Source/Runtime/StorageServerClient/Private/StorageServerPlatformFile.h
--- a/Engine/Source/Runtime/StorageServerClient/Private/StorageServerPlatformFile.h
+++ b/Engine/Source/Runtime/StorageServerClient/Private/StorageServerPlatformFile.h
@@ -105,6 +105,7 @@ private:
 	FStringView ServerProjectDirView = FStringView(TEXT("/{project}/"));
 	TUniquePtr<FStorageServerConnection> Connection;
 	FStorageServerFileSystemTOC ServerToc;
+	mutable TArray<FString> HostAddrs;
 };
 
 #endif
diff --git a/Engine/Source/Runtime/StorageServerClient/Public/StorageServerConnection.h b/Engine/Source/Runtime/StorageServerClient/Public/StorageServerConnection.h
--- a/Engine/Source/Runtime/StorageServerClient/Public/StorageServerConnection.h
+++ b/Engine/Source/Runtime/StorageServerClient/Public/StorageServerConnection.h
@@ -112,7 +112,7 @@ public:
 	STORAGESERVERCLIENT_API FStorageServerConnection();
 	STORAGESERVERCLIENT_API ~FStorageServerConnection();
 
-	STORAGESERVERCLIENT_API bool Initialize(const TCHAR* Host, int32 Port, const TCHAR* ProjectNameOverride = nullptr, const TCHAR* PlatformNameOverride = nullptr);
+	STORAGESERVERCLIENT_API bool Initialize(TArrayView<const FString> HostAddresses, int32 Port, const TCHAR* ProjectNameOverride = nullptr, const TCHAR* PlatformNameOverride = nullptr);
 
 	STORAGESERVERCLIENT_API void FileManifestRequest(TFunctionRef<void(FIoChunkId Id, FStringView Path)> Callback);
 	STORAGESERVERCLIENT_API int64 FileSizeRequest(int32 FileIndex);
@@ -121,12 +121,14 @@ public:
 	STORAGESERVERCLIENT_API bool ReadChunkRequest(const FIoChunkId& ChunkId, uint64 Offset, uint64 Size, TFunctionRef<void(FStorageServerResponse&)> OnResponse);
 	STORAGESERVERCLIENT_API FStorageServerChunkBatchRequest NewChunkBatchRequest();
 
+	STORAGESERVERCLIENT_API FString GetHostAddr() const;
+
 private:
 	friend class FStorageServerRequest;
 	friend class FStorageServerResponse;
 	friend class FStorageServerChunkBatchRequest;
 
-	int32 HandshakeRequest();
+	int32 HandshakeRequest(TArrayView<const TSharedPtr<FInternetAddr>> HostAddresses);
 	FSocket* AcquireSocket();
 	void ReleaseSocket(FSocket* Socket, bool bKeepAlive);
 
