Zen: Changed file chunk ID to be based on filename hash instead of a monotonically increasing index

#rnx

diff --git a/Engine/Source/Developer/IoStoreUtilities/Private/ZenFileSystemManifest.cpp b/Engine/Source/Developer/IoStoreUtilities/Private/ZenFileSystemManifest.cpp
--- a/Engine/Source/Developer/IoStoreUtilities/Private/ZenFileSystemManifest.cpp
+++ b/Engine/Source/Developer/IoStoreUtilities/Private/ZenFileSystemManifest.cpp
@@ -103,11 +103,13 @@ int32 FZenFileSystemManifest::Generate()
 					{
 						FStringView RelativePath = InnerFileNameOrDirectory;
 						RelativePath.RightChopInline(LocalDirectory.Len() + 1);
+						FString ClientPath = FPaths::Combine(ClientDirectory, RelativePath.GetData());
+						const FIoChunkId FileChunkId = CreateExternalFileChunkId(ClientPath);
 
 						AddManifestEntry(
-							CreateExternalFileChunkId(0, Entries.Num() + 1),
+							FileChunkId,
 							FPaths::Combine(ServerRelativeDirectory, RelativePath.GetData()),
-							FPaths::Combine(ClientDirectory, RelativePath.GetData()));
+							MoveTemp(ClientPath));
 					}
 
 					return true;
@@ -170,10 +172,11 @@ const FZenFileSystemManifest::FManifestEntry& FZenFileSystemManifest::CreateMani
 		FPaths::MakePathRelativeTo(ServerRelativeDirectory, *FPaths::RootDir());
 		ServerRelativeDirectory = TEXT("/") + ServerRelativeDirectory;
 
-		return AddManifestEntry(
-			CreateExternalFileChunkId(0, Entries.Num() + 1),
-			FPaths::Combine(ServerRelativeDirectory, RelativePath.GetData()),
-			FPaths::Combine(ClientDirectory, RelativePath.GetData()));
+		FString ServerPath = FPaths::Combine(ServerRelativeDirectory, RelativePath.GetData());
+		FString ClientPath = FPaths::Combine(ClientDirectory, RelativePath.GetData());
+		const FIoChunkId FileChunkId = CreateExternalFileChunkId(ClientPath);
+
+		return AddManifestEntry(FileChunkId, MoveTemp(ServerPath), MoveTemp(ClientPath));
 	};
 
 	if (Filename.StartsWith(CookedEngineDirectory))
diff --git a/Engine/Source/Runtime/Core/Private/IO/IoDispatcher.cpp b/Engine/Source/Runtime/Core/Private/IO/IoDispatcher.cpp
--- a/Engine/Source/Runtime/Core/Private/IO/IoDispatcher.cpp
+++ b/Engine/Source/Runtime/Core/Private/IO/IoDispatcher.cpp
@@ -19,10 +19,10 @@
 #include "Containers/Ticker.h"
 #include "IO/IoDispatcherBackend.h"
 #include "IO/IoDispatcherFileBackendTypes.h"
+#include "Hash/Blake3.h"
 
 DEFINE_LOG_CATEGORY(LogIoDispatcher);
 
-
 const FIoChunkId FIoChunkId::InvalidChunkId = FIoChunkId::CreateEmptyId();
 
 TUniquePtr<FIoDispatcher> GIoDispatcher;
@@ -1027,3 +1027,32 @@ FIoRequest::UpdatePriority(uint32 NewPriority)
 	Impl->Dispatcher.Reprioritize(Impl);
 }
 
+FIoChunkId CreateExternalFileChunkId(const FStringView Filename)
+{
+	check(Filename.Len() > 0);
+
+	TArray<TCHAR, TInlineAllocator<FName::StringBufferSize>> Buffer;
+	Buffer.SetNum(Filename.Len());
+
+	for (int32 Idx = 0, Len = Filename.Len(); Idx < Len; Idx++)
+	{
+		TCHAR Char = TChar<TCHAR>::ToLower(Filename[Idx]);
+		if (Char == TEXT('\\'))
+		{
+			Char = TEXT('/');
+		}
+
+		Buffer[Idx] = Char;
+	}
+
+	const FBlake3Hash Hash = FBlake3::HashBuffer(FMemoryView(Buffer.GetData(), Buffer.Num() * Buffer.GetTypeSize()));
+
+	uint8 Id[12] = {0};
+	FMemory::Memcpy(Id, Hash.GetBytes(), 11);
+	Id[11] = uint8(EIoChunkType::ExternalFile);
+
+	FIoChunkId ChunkId;
+	ChunkId.Set(Id, sizeof(Id));
+
+	return ChunkId;
+}
diff --git a/Engine/Source/Runtime/Core/Public/IO/IoDispatcher.h b/Engine/Source/Runtime/Core/Public/IO/IoDispatcher.h
--- a/Engine/Source/Runtime/Core/Public/IO/IoDispatcher.h
+++ b/Engine/Source/Runtime/Core/Public/IO/IoDispatcher.h
@@ -738,7 +738,6 @@ private:
 	uint8	Id[12];
 };
 
-
 /**
  * Creates a chunk identifier (generic -- prefer specialized versions where possible)
  */
@@ -758,44 +757,7 @@ static FIoChunkId CreateIoChunkId(uint64 ChunkId, uint16 ChunkIndex, EIoChunkTyp
 	return IoChunkId;
 }
 
-/**
- * Create an "external file" chunk identifier
- * 
- * Layout (the index is big endian for better lexical sort behavior):
- * 
- *   00  01  02  03  04  05  06  07  08  09  10  11
- * +---+---+---+---+---+---+---+---+---+---+---+---+
- * |      ContainerId              | FileIndex | 7 |
- * +---+---+---+---+---+---+---+---+---+---+---+---+
- * 
- */
-static inline FIoChunkId CreateExternalFileChunkId(uint64 ContainerId, uint32 FileIndex)
-{
-	check(FileIndex < (1 << 24));
-
-	uint8 Data[12] = { 0 };
-
-	Data[11] = uint8(EIoChunkType::ExternalFile);
-
-	memcpy(&Data[0], &ContainerId, sizeof ContainerId);
-
-	const uint32 NoFileIndex = NETWORK_ORDER32(FileIndex);
-	memcpy(&Data[8], reinterpret_cast<const uint8*>(&NoFileIndex) + 1, 3);
-
-	FIoChunkId IoChunkId;
-	IoChunkId.Set(Data, 12);
-
-	return IoChunkId;
-}
-
-inline uint32 FileIndexFromChunkId(FIoChunkId ChunkId)
-{
-	// See encoding details in CreateExternalFileChunkId()
-	uint32 FileIndex = 0;
-	memcpy(&FileIndex, ChunkId.GetData() + 7, sizeof FileIndex);
-	FileIndex = NETWORK_ORDER32(FileIndex) & 0xffFFff;
-	return FileIndex;
-}
+CORE_API FIoChunkId CreateExternalFileChunkId(const FStringView Filename);
 
 //////////////////////////////////////////////////////////////////////////
 
diff --git a/Engine/Source/Runtime/StorageServerClient/Private/StorageServerConnection.cpp b/Engine/Source/Runtime/StorageServerClient/Private/StorageServerConnection.cpp
--- a/Engine/Source/Runtime/StorageServerClient/Private/StorageServerConnection.cpp
+++ b/Engine/Source/Runtime/StorageServerClient/Private/StorageServerConnection.cpp
@@ -403,75 +403,6 @@ void FStorageServerConnection::FileManifestRequest(TFunctionRef<void(FIoChunkId
 	}
 }
 
-void FileIndexToChunkId(FAnsiStringBuilderBase& OutString, int32 FileIndex)
-{
-	FIoChunkId Chunk = CreateExternalFileChunkId(0, FileIndex);
-	OutString << Chunk;
-}
-
-int64 FStorageServerConnection::FileSizeRequest(int32 FileIndex)
-{
-	TAnsiStringBuilder<256> ResourceBuilder;
-	ResourceBuilder.Append(OplogPath);
-	ResourceBuilder << "/";
-	FileIndexToChunkId(ResourceBuilder, FileIndex);
-
-	FStorageServerRequest Request("HEAD", *ResourceBuilder, Hostname);
-	FSocket* Socket = Request.Send(*this);
-	if (!Socket)
-	{
-		UE_LOG(LogStorageServerConnection, Fatal, TEXT("Failed to send file size request to storage server at %s."), *ServerAddr->ToString(true));
-		return -1;
-	}
-	FStorageServerResponse Response(*this, *Socket);
-	if (Response.IsOk())
-	{
-		return Response.ContentLength;
-	}
-	else if (Response.GetErrorCode() == 404)
-	{
-		return -1;
-	}
-	else
-	{
-		UE_LOG(LogStorageServerConnection, Fatal, TEXT("Failed to get file size from storage server at %s. '%s'"), *ServerAddr->ToString(true), *Response.GetErrorMessage());
-	}
-	return -1;
-}
-
-bool FStorageServerConnection::ReadFileRequest(int32 FileIndex, uint64 Offset, uint64 Size, TFunctionRef<void(FStorageServerResponse&)> OnResponse)
-{
-	TAnsiStringBuilder<256> ResourceBuilder;
-	ResourceBuilder.Append(OplogPath);
-	ResourceBuilder << "/";
-	FileIndexToChunkId(ResourceBuilder, FileIndex);
-
-	ResourceBuilder.Appendf("?offset=%" UINT64_FMT "&size=%" UINT64_FMT, Offset, Size);
-
-	FStorageServerRequest Request = FStorageServerRequest("GET", *ResourceBuilder, Hostname);
-	FSocket* Socket = Request.Send(*this);
-	if (!Socket)
-	{
-		UE_LOG(LogStorageServerConnection, Fatal, TEXT("Failed to send file read request to storage server at %s."), *ServerAddr->ToString(true));
-		return false;
-	}
-	FStorageServerResponse Response(*this, *Socket);
-	if (Response.IsOk())
-	{
-		OnResponse(Response);
-		return true;
-	}
-	else if (Response.GetErrorCode() == 404)
-	{
-		return false;
-	}
-	else
-	{
-		UE_LOG(LogStorageServerConnection, Fatal, TEXT("Failed to read file from storage server at %s."), *ServerAddr->ToString(true), *Response.GetErrorMessage());
-	}
-	return false;
-}
-
 int64 FStorageServerConnection::ChunkSizeRequest(const FIoChunkId& ChunkId)
 {
 	TAnsiStringBuilder<256> ResourceBuilder;
diff --git a/Engine/Source/Runtime/StorageServerClient/Private/StorageServerPlatformFile.cpp b/Engine/Source/Runtime/StorageServerClient/Private/StorageServerPlatformFile.cpp
--- a/Engine/Source/Runtime/StorageServerClient/Private/StorageServerPlatformFile.cpp
+++ b/Engine/Source/Runtime/StorageServerClient/Private/StorageServerPlatformFile.cpp
@@ -51,20 +51,25 @@ FStorageServerFileSystemTOC::FDirectory* FStorageServerFileSystemTOC::AddDirecto
 	return Directory;
 }
 
-void FStorageServerFileSystemTOC::AddFile(FStringView PathView, int32 Index)
+void FStorageServerFileSystemTOC::AddFile(const FIoChunkId& FileChunkId, FStringView PathView)
 {
 	FWriteScopeLock _(TocLock);
 
-	FString Path(PathView);
-	FilePathToIndexMap.Add(Path, Index);
-	FileIndexToPathMap.Add(Index, Path);
-	FString DirectoryPath = FPaths::GetPath(Path);
+	const int32 FileIndex = Files.Num();
+	
+	FFile& NewFile = Files.AddDefaulted_GetRef();
+	NewFile.FileChunkId = FileChunkId;
+	NewFile.FilePath = PathView;
+	
+	FilePathToIndexMap.Add(NewFile.FilePath, FileIndex);
+	
+	FString DirectoryPath = FPaths::GetPath(NewFile.FilePath);
 	FDirectory* Directory = Directories.FindRef(DirectoryPath);
 	if (!Directory)
 	{
 		Directory = AddDirectoriesRecursive(DirectoryPath);
 	}
-	Directory->Files.Add(Index);
+	Directory->Files.Add(FileIndex);
 }
 
 bool FStorageServerFileSystemTOC::FileExists(const FString& Path)
@@ -79,13 +84,17 @@ bool FStorageServerFileSystemTOC::DirectoryExists(const FString& Path)
 	return Directories.Contains(Path);
 }
 
-int32* FStorageServerFileSystemTOC::FindFileIndex(const FString& Path)
+const FIoChunkId* FStorageServerFileSystemTOC::GetFileChunkId(const FString& Path)
 {
 	FReadScopeLock _(TocLock);
-	return FilePathToIndexMap.Find(Path);
+	if (const int32* FileIndex = FilePathToIndexMap.Find(Path))
+	{
+		return &Files[*FileIndex].FileChunkId;
+	}
+	return nullptr;
 }
 
-bool FStorageServerFileSystemTOC::IterateDirectory(const FString& Path, TFunctionRef<bool(int32, const TCHAR*)> Callback)
+bool FStorageServerFileSystemTOC::IterateDirectory(const FString& Path, TFunctionRef<bool(const FIoChunkId&, const TCHAR*)> Callback)
 {
 	UE_LOG(LogStorageServerPlatformFile, Verbose, TEXT("IterateDirectory '%s'"), *Path);
 
@@ -98,15 +107,15 @@ bool FStorageServerFileSystemTOC::IterateDirectory(const FString& Path, TFunctio
 	}
 	for (int32 FileIndex : Directory->Files)
 	{
-		const FString* FindFilePath = FileIndexToPathMap.Find(FileIndex);
-		if (FindFilePath && !Callback(FileIndex, **FindFilePath))
+		const FFile& File = Files[FileIndex];
+		if (!Callback(File.FileChunkId, *File.FilePath))
 		{
 			return false;
 		}
 	}
 	for (const FString& ChildDirectoryPath : Directory->Directories)
 	{
-		if (!Callback(-1, *ChildDirectoryPath))
+		if (!Callback(FIoChunkId(), *ChildDirectoryPath))
 		{
 			return false;
 		}
@@ -122,8 +131,8 @@ class FStorageServerFileHandle
 		BufferSize = 64 << 10
 	};
 	FStorageServerPlatformFile& Owner;
+	FIoChunkId FileChunkId;
 	FString Filename;
-	int32 FileIndex;
 	int64 FilePos = 0;
 	int64 FileSize = -1;
 	int64 BufferStart = -1;
@@ -131,10 +140,10 @@ class FStorageServerFileHandle
 	uint8 Buffer[BufferSize];
 
 public:
-	FStorageServerFileHandle(FStorageServerPlatformFile& InOwner, const TCHAR* InFilename, int32 InFileIndex)
+	FStorageServerFileHandle(FStorageServerPlatformFile& InOwner, FIoChunkId InFileChunkId, const TCHAR* InFilename)
 		: Owner(InOwner)
+		, FileChunkId(InFileChunkId)
 		, Filename(InFilename)
-		, FileIndex(InFileIndex)
 	{
 	}
 
@@ -146,7 +155,7 @@ public:
 	{
 		if (FileSize < 0)
 		{
-			FFileStatData FileStatData = Owner.SendGetStatDataMessage(FileIndex);
+			const FFileStatData FileStatData = Owner.SendGetStatDataMessage(FileChunkId);
 			if (FileStatData.bIsValid)
 			{
 				FileSize = FileStatData.FileSize;
@@ -185,7 +194,7 @@ public:
 
 		if (BytesToRead > BufferSize)
 		{
-			int64 BytesRead = Owner.SendReadMessage(Destination, FileIndex, FilePos, BytesToRead);
+			const int64 BytesRead = Owner.SendReadMessage(Destination, FileChunkId, FilePos, BytesToRead);
 			if (BytesRead == BytesToRead)
 			{
 				FilePos += BytesRead;
@@ -195,7 +204,7 @@ public:
 
 		if (FilePos < BufferStart || BufferEnd < FilePos + BytesToRead)
 		{
-			int64 BytesRead = Owner.SendReadMessage(Buffer, FileIndex, FilePos, BufferSize);
+			const int64 BytesRead = Owner.SendReadMessage(Buffer, FileChunkId, FilePos, BufferSize);
 			BufferStart = FilePos;
 			BufferEnd = BufferStart + BytesRead;
 		}
@@ -322,10 +331,9 @@ FDateTime FStorageServerPlatformFile::GetTimeStamp(const TCHAR* Filename)
 	TStringBuilder<1024> StorageServerFilename;
 	if (MakeStorageServerPath(Filename, StorageServerFilename))
 	{
-		int32* FindFileIndex = ServerToc.FindFileIndex(*StorageServerFilename);
-		if (FindFileIndex)
+		if (const FIoChunkId* FileChunkId = ServerToc.GetFileChunkId(*StorageServerFilename))
 		{
-			FFileStatData FileStatData = SendGetStatDataMessage(*FindFileIndex);
+			const FFileStatData FileStatData = SendGetStatDataMessage(*FileChunkId);
 			check(FileStatData.bIsValid);
 			return FileStatData.ModificationTime;
 		}
@@ -338,10 +346,9 @@ FDateTime FStorageServerPlatformFile::GetAccessTimeStamp(const TCHAR* Filename)
 	TStringBuilder<1024> StorageServerFilename;
 	if (MakeStorageServerPath(Filename, StorageServerFilename))
 	{
-		int32* FindFileIndex = ServerToc.FindFileIndex(*StorageServerFilename);
-		if (FindFileIndex)
+		if (const FIoChunkId* FileChunkId = ServerToc.GetFileChunkId(*StorageServerFilename))
 		{
-			FFileStatData FileStatData = SendGetStatDataMessage(*FindFileIndex);
+			const FFileStatData FileStatData = SendGetStatDataMessage(*FileChunkId);
 			check(FileStatData.bIsValid);
 			return FileStatData.AccessTime;
 		}
@@ -354,10 +361,9 @@ int64 FStorageServerPlatformFile::FileSize(const TCHAR* Filename)
 	TStringBuilder<1024> StorageServerFilename;
 	if (MakeStorageServerPath(Filename, StorageServerFilename))
 	{
-		int32* FindFileIndex = ServerToc.FindFileIndex(*StorageServerFilename);
-		if (FindFileIndex)
+		if (const FIoChunkId* FileChunkId = ServerToc.GetFileChunkId(*StorageServerFilename))
 		{
-			FFileStatData FileStatData = SendGetStatDataMessage(*FindFileIndex);
+			const FFileStatData FileStatData = SendGetStatDataMessage(*FileChunkId);
 			check(FileStatData.bIsValid);
 			return FileStatData.FileSize;
 		}
@@ -380,10 +386,9 @@ FFileStatData FStorageServerPlatformFile::GetStatData(const TCHAR* FilenameOrDir
 	TStringBuilder<1024> StorageServerFilenameOrDirectory;
 	if (MakeStorageServerPath(FilenameOrDirectory, StorageServerFilenameOrDirectory))
 	{
-		int32* FindFileIndex = ServerToc.FindFileIndex(*StorageServerFilenameOrDirectory);
-		if (FindFileIndex)
+		if (const FIoChunkId* FileChunkId = ServerToc.GetFileChunkId(*StorageServerFilenameOrDirectory))
 		{
-			return SendGetStatDataMessage(*FindFileIndex);
+			return SendGetStatDataMessage(*FileChunkId);
 		}
 		else if (ServerToc.DirectoryExists(*StorageServerFilenameOrDirectory))
 		{
@@ -399,9 +404,9 @@ FFileStatData FStorageServerPlatformFile::GetStatData(const TCHAR* FilenameOrDir
 	return LowerLevel->GetStatData(FilenameOrDirectory);
 }
 
-IFileHandle* FStorageServerPlatformFile::InternalOpenFile(int32 FileIndex, const TCHAR* LocalFilename)
+IFileHandle* FStorageServerPlatformFile::InternalOpenFile(const FIoChunkId& FileChunkId, const TCHAR* LocalFilename)
 {
-	return new FStorageServerFileHandle(*this, LocalFilename, FileIndex);
+	return new FStorageServerFileHandle(*this, FileChunkId, LocalFilename);
 }
 
 IFileHandle* FStorageServerPlatformFile::OpenRead(const TCHAR* Filename, bool bAllowWrite)
@@ -409,10 +414,9 @@ IFileHandle* FStorageServerPlatformFile::OpenRead(const TCHAR* Filename, bool bA
 	TStringBuilder<1024> StorageServerFilename;
 	if (MakeStorageServerPath(Filename, StorageServerFilename))
 	{
-		int32* FindFileIndex = ServerToc.FindFileIndex(*StorageServerFilename);
-		if (FindFileIndex)
+		if (const FIoChunkId* FileChunkId = ServerToc.GetFileChunkId(*StorageServerFilename))
 		{
-			return InternalOpenFile(*FindFileIndex, Filename);
+			return InternalOpenFile(*FileChunkId, Filename);
 		}
 	}
 	return LowerLevel->OpenRead(Filename, bAllowWrite);
@@ -424,12 +428,13 @@ bool FStorageServerPlatformFile::IterateDirectory(const TCHAR* Directory, IPlatf
 	bool bResult = false;
 	if (MakeStorageServerPath(Directory, StorageServerDirectory) && ServerToc.DirectoryExists(*StorageServerDirectory))
 	{
-		bResult |= ServerToc.IterateDirectory(*StorageServerDirectory, [this, &Visitor](int32 FileIndex, const TCHAR* FilenameOrDirectory)
+		bResult |= ServerToc.IterateDirectory(*StorageServerDirectory, [this, &Visitor](const FIoChunkId& FileChunkId, const TCHAR* FilenameOrDirectory)
 		{
 			TStringBuilder<1024> LocalPath;
 			bool bConverted = MakeLocalPath(FilenameOrDirectory, LocalPath);
 			check(bConverted);
-			return Visitor.Visit(*LocalPath, FileIndex < 0);
+			const bool bDirectory = !FileChunkId.IsValid();
+			return Visitor.Visit(*LocalPath, bDirectory);
 		});
 	}
 	else
@@ -445,15 +450,15 @@ bool FStorageServerPlatformFile::IterateDirectoryStat(const TCHAR* Directory, FD
 	bool bResult = false;
 	if (MakeStorageServerPath(Directory, StorageServerDirectory) && ServerToc.DirectoryExists(*StorageServerDirectory))
 	{
-		bResult |= ServerToc.IterateDirectory(*StorageServerDirectory, [this, &Visitor](int32 FileIndex, const TCHAR* ServerFilenameOrDirectory)
+		bResult |= ServerToc.IterateDirectory(*StorageServerDirectory, [this, &Visitor](const FIoChunkId& FileChunkId, const TCHAR* ServerFilenameOrDirectory)
 		{
 			TStringBuilder<1024> LocalPath;
 			bool bConverted = MakeLocalPath(ServerFilenameOrDirectory, LocalPath);
 			check(bConverted);
 			FFileStatData FileStatData;
-			if (FileIndex >= 0)
+			if (FileChunkId.IsValid())
 			{
-				FileStatData = SendGetStatDataMessage(FileIndex);
+				FileStatData = SendGetStatDataMessage(FileChunkId);
 				check(FileStatData.bIsValid);
 			}
 			else
@@ -517,8 +522,7 @@ bool FStorageServerPlatformFile::MoveFile(const TCHAR* To, const TCHAR* From)
 	TStringBuilder<1024> StorageServerFrom;
 	if (MakeStorageServerPath(From, StorageServerFrom))
 	{
-		int32* FindFromFileIndex = ServerToc.FindFileIndex(*StorageServerFrom);
-		if (FindFromFileIndex)
+		if (const FIoChunkId* FromFileChunkId = ServerToc.GetFileChunkId(*StorageServerFrom))
 		{
 			TUniquePtr<IFileHandle> ToFile(LowerLevel->OpenWrite(To, false, false));
 			if (!ToFile)
@@ -526,7 +530,7 @@ bool FStorageServerPlatformFile::MoveFile(const TCHAR* To, const TCHAR* From)
 				return false;
 			}
 
-			TUniquePtr<IFileHandle> FromFile(InternalOpenFile(*FindFromFileIndex, *StorageServerFrom));
+			TUniquePtr<IFileHandle> FromFile(InternalOpenFile(*FromFileChunkId, *StorageServerFrom));
 			if (!FromFile)
 			{
 				return false;
@@ -657,16 +661,16 @@ bool FStorageServerPlatformFile::SendGetFileListMessage()
 	
 	Connection->FileManifestRequest([&](FIoChunkId Id, FStringView Path)
 	{
-		ServerToc.AddFile(Path, FileIndexFromChunkId(Id));
+		ServerToc.AddFile(Id, Path);
 	});
 
 	return true;
 }
 
-FFileStatData FStorageServerPlatformFile::SendGetStatDataMessage(int32 FileIndex)
+FFileStatData FStorageServerPlatformFile::SendGetStatDataMessage(const FIoChunkId& FileChunkId)
 {
 	TRACE_CPUPROFILER_EVENT_SCOPE(StorageServerPlatformFileGetStatData);
-	int64 FileSize = Connection->FileSizeRequest(FileIndex);
+	const int64 FileSize = Connection->ChunkSizeRequest(FileChunkId);
 	if (FileSize < 0)
 	{
 		return FFileStatData();
@@ -679,11 +683,11 @@ FFileStatData FStorageServerPlatformFile::SendGetStatDataMessage(int32 FileIndex
 	return FFileStatData(CreationTime, AccessTime, ModificationTime, FileSize, false, true);
 }
 
-int64 FStorageServerPlatformFile::SendReadMessage(uint8* Destination, int32 FileIndex, int64 Offset, int64 BytesToRead)
+int64 FStorageServerPlatformFile::SendReadMessage(uint8* Destination, const FIoChunkId& FileChunkId, int64 Offset, int64 BytesToRead)
 {
 	TRACE_CPUPROFILER_EVENT_SCOPE(StorageServerPlatformFileRead);
 	int64 BytesRead = 0;
-	Connection->ReadFileRequest(FileIndex, Offset, BytesToRead, [Destination, &BytesRead](FStorageServerResponse& Response)
+	Connection->ReadChunkRequest(FileChunkId, Offset, BytesToRead, [Destination, &BytesRead](FStorageServerResponse& Response)
 	{
 		BytesRead = Response.TotalSize();
 		Response.Serialize(Destination, Response.TotalSize());
@@ -714,7 +718,7 @@ void FStorageServerPlatformFile::OnCookOnTheFlyMessage(const UE::Cook::FCookOnTh
 			for (int32 Idx = 0, Num = Filenames.Num(); Idx < Num; ++Idx)
 			{
 				UE_LOG(LogCookOnTheFly, Verbose, TEXT("Adding file '%s'"), *Filenames[Idx]);
-				ServerToc.AddFile(Filenames[Idx], FileIndexFromChunkId(ChunkIds[Idx]));
+				ServerToc.AddFile(ChunkIds[Idx], Filenames[Idx]);
 			}
 
 			break;
diff --git a/Engine/Source/Runtime/StorageServerClient/Private/StorageServerPlatformFile.h b/Engine/Source/Runtime/StorageServerClient/Private/StorageServerPlatformFile.h
--- a/Engine/Source/Runtime/StorageServerClient/Private/StorageServerPlatformFile.h
+++ b/Engine/Source/Runtime/StorageServerClient/Private/StorageServerPlatformFile.h
@@ -5,6 +5,7 @@
 #include "CoreMinimal.h"
 #include "GenericPlatform/GenericPlatformFile.h"
 #include "Containers/StringView.h"
+#include "IO/IoDispatcher.h"
 
 #if !UE_BUILD_SHIPPING
 
@@ -22,11 +23,11 @@ class FStorageServerFileSystemTOC
 {
 public:
 	~FStorageServerFileSystemTOC();
-	void AddFile(FStringView Path, int32 Index);
+	void AddFile(const FIoChunkId& FileChunkId, FStringView Path);
 	bool FileExists(const FString& Path);
 	bool DirectoryExists(const FString& Path);
-	int32* FindFileIndex(const FString& Path);
-	bool IterateDirectory(const FString& Path, TFunctionRef<bool(int32, const TCHAR*)> Callback);
+	const FIoChunkId* GetFileChunkId(const FString& Path);
+	bool IterateDirectory(const FString& Path, TFunctionRef<bool(const FIoChunkId&, const TCHAR*)> Callback);
 
 private:
 	struct FDirectory
@@ -35,12 +36,18 @@ private:
 		TArray<int32> Files;
 	};
 
+	struct FFile
+	{
+		FIoChunkId FileChunkId;
+		FString FilePath;
+	};
+
 	FDirectory* AddDirectoriesRecursive(const FString& DirectoryPath);
 
 	FDirectory Root;
 	TMap<FString, FDirectory*> Directories;
-	TMap<int32, FString> FileIndexToPathMap;
 	TMap<FString, int32> FilePathToIndexMap;
+	TArray<FFile> Files;
 	FRWLock TocLock;
 };
 
@@ -92,10 +99,10 @@ private:
 
 	bool MakeStorageServerPath(const TCHAR* LocalFilenameOrDirectory, FStringBuilderBase& OutPath) const;
 	bool MakeLocalPath(const TCHAR* ServerFilenameOrDirectory, FStringBuilderBase& OutPath) const;
-	IFileHandle* InternalOpenFile(int32 FileIndex, const TCHAR* LocalFilename);
+	IFileHandle* InternalOpenFile(const FIoChunkId& FileChunkId, const TCHAR* LocalFilename);
 	bool SendGetFileListMessage();
-	FFileStatData SendGetStatDataMessage(int32 FileIndex);
-	int64 SendReadMessage(uint8* Destination, int32 FileIndex, int64 Offset, int64 BytesToRead);
+	FFileStatData SendGetStatDataMessage(const FIoChunkId& FileChunkId);
+	int64 SendReadMessage(uint8* Destination, const FIoChunkId& FileChunkId, int64 Offset, int64 BytesToRead);
 #if WITH_COTF
 	void OnCookOnTheFlyMessage(const UE::Cook::FCookOnTheFlyMessage& Message);
 #endif
diff --git a/Engine/Source/Runtime/StorageServerClient/Public/StorageServerConnection.h b/Engine/Source/Runtime/StorageServerClient/Public/StorageServerConnection.h
--- a/Engine/Source/Runtime/StorageServerClient/Public/StorageServerConnection.h
+++ b/Engine/Source/Runtime/StorageServerClient/Public/StorageServerConnection.h
@@ -115,8 +115,6 @@ public:
 	STORAGESERVERCLIENT_API bool Initialize(TArrayView<const FString> HostAddresses, int32 Port, const TCHAR* ProjectNameOverride = nullptr, const TCHAR* PlatformNameOverride = nullptr);
 
 	STORAGESERVERCLIENT_API void FileManifestRequest(TFunctionRef<void(FIoChunkId Id, FStringView Path)> Callback);
-	STORAGESERVERCLIENT_API int64 FileSizeRequest(int32 FileIndex);
-	STORAGESERVERCLIENT_API bool ReadFileRequest(int32 FileIndex, uint64 Offset, uint64 Size, TFunctionRef<void(FStorageServerResponse&)> OnResponse);
 	STORAGESERVERCLIENT_API int64 ChunkSizeRequest(const FIoChunkId& ChunkId);
 	STORAGESERVERCLIENT_API bool ReadChunkRequest(const FIoChunkId& ChunkId, uint64 Offset, uint64 Size, TFunctionRef<void(FStorageServerResponse&)> OnResponse);
 	STORAGESERVERCLIENT_API FStorageServerChunkBatchRequest NewChunkBatchRequest();
