Added missing include
#rnx

diff --git a/Engine/Source/Runtime/StorageServerClient/Private/StorageServerPlatformFile.cpp b/Engine/Source/Runtime/StorageServerClient/Private/StorageServerPlatformFile.cpp
--- a/Engine/Source/Runtime/StorageServerClient/Private/StorageServerPlatformFile.cpp
+++ b/Engine/Source/Runtime/StorageServerClient/Private/StorageServerPlatformFile.cpp
@@ -12,6 +12,7 @@
 #include "Algo/Replace.h"
 #include "StorageServerPackageStore.h"
 #include "CookOnTheFlyPackageStore.h"
+#include "Misc/CoreDelegates.h"
 
 #if WITH_COTF
 #include "Modules/ModuleManager.h"
