[StorageServer] Change fatal error to log when connecting to potential storage server addresses.  This allows connection attempts to continue, iterating through the specified list of host addresses

diff --git a/Engine/Source/Runtime/StorageServerClient/Private/StorageServerConnection.cpp b/Engine/Source/Runtime/StorageServerClient/Private/StorageServerConnection.cpp
--- a/Engine/Source/Runtime/StorageServerClient/Private/StorageServerConnection.cpp
+++ b/Engine/Source/Runtime/StorageServerClient/Private/StorageServerConnection.cpp
@@ -81,6 +81,12 @@ FSocket* FStorageServerRequest::Send(FStorageServerConnection& Owner)
 	while (Attempts++ < 10)
 	{
 		FSocket* Socket = Owner.AcquireSocket();
+
+		if (!Socket)
+		{
+			break;
+		}
+
 		if (Send(Socket, reinterpret_cast<const uint8*>(HeaderBuffer.GetData()), HeaderBuffer.Len()) &&
 			Send(Socket, BodyBuffer.GetData(), BodyBuffer.Num()))
 		{
@@ -522,7 +528,7 @@ FSocket* FStorageServerConnection::AcquireSocket()
 		delete Socket;
 	}
 
-	UE_LOG(LogStorageServerConnection, Fatal, TEXT("Failed to connect to storage server at %s."), *ServerAddr->ToString(true));
+	UE_LOG(LogStorageServerConnection, Log, TEXT("Failed to connect to storage server at %s."), *ServerAddr->ToString(true));
 
 	return nullptr;
 }
