Non-Unity fixes

diff --git a/Engine/Source/Runtime/Engine/Private/SkeletalMeshDeformerHelpers.cpp b/Engine/Source/Runtime/Engine/Private/SkeletalMeshDeformerHelpers.cpp
--- a/Engine/Source/Runtime/Engine/Private/SkeletalMeshDeformerHelpers.cpp
+++ b/Engine/Source/Runtime/Engine/Private/SkeletalMeshDeformerHelpers.cpp
@@ -1,7 +1,4 @@
 // Copyright Epic Games, Inc. All Rights Reserved.
-
-#pragma once
-
 #include "SkeletalMeshDeformerHelpers.h"
 
 #include "GPUSkinVertexFactory.h"
diff --git a/Engine/Source/Runtime/StorageServerClient/Private/StorageServerConnection.cpp b/Engine/Source/Runtime/StorageServerClient/Private/StorageServerConnection.cpp
--- a/Engine/Source/Runtime/StorageServerClient/Private/StorageServerConnection.cpp
+++ b/Engine/Source/Runtime/StorageServerClient/Private/StorageServerConnection.cpp
@@ -15,6 +15,7 @@
 #include "Serialization/CompactBinarySerialization.h"
 #include "Memory/SharedBuffer.h"
 #include "ProfilingDebugging/CountersTrace.h"
+#include "HAL/FileManager.h"
 
 #if !UE_BUILD_SHIPPING
 
