Zen - fix bug with missing offset argument

#rnx

diff --git a/Engine/Source/Runtime/StorageServerClient/Private/StorageServerPlatformFile.cpp b/Engine/Source/Runtime/StorageServerClient/Private/StorageServerPlatformFile.cpp
--- a/Engine/Source/Runtime/StorageServerClient/Private/StorageServerPlatformFile.cpp
+++ b/Engine/Source/Runtime/StorageServerClient/Private/StorageServerPlatformFile.cpp
@@ -688,9 +688,9 @@ int64 FStorageServerPlatformFile::SendReadMessage(uint8* Destination, const FIoC
 {
 	TRACE_CPUPROFILER_EVENT_SCOPE(StorageServerPlatformFileRead);
 	int64 BytesRead = 0;
-	Connection->ReadChunkRequest(FileChunkId, Offset, BytesToRead, [Destination, BytesToRead, &BytesRead](FStorageServerResponse& Response)
+	Connection->ReadChunkRequest(FileChunkId, Offset, BytesToRead, [Destination, Offset, BytesToRead, &BytesRead](FStorageServerResponse& Response)
 	{
-		BytesRead = Response.SerializeChunkTo(MakeMemoryView(Destination, BytesToRead));
+		BytesRead = Response.SerializeChunkTo(MakeMemoryView(Destination, BytesToRead), Offset);
 	});
 	return BytesRead;
 }
