Check local cooked for files that are to be read externally

Currently there are some instances where data must exist on the local file system in order for it to be read by a third party API (Mac's Metal shader libraries being one such example). This change overlays files in Saved/Cooked and checks in there for any external-read path requested that is and prefixed with ../../../
#rnx

diff --git a/Engine/Source/Runtime/StorageServerClient/Private/StorageServerPlatformFile.cpp b/Engine/Source/Runtime/StorageServerClient/Private/StorageServerPlatformFile.cpp
--- a/Engine/Source/Runtime/StorageServerClient/Private/StorageServerPlatformFile.cpp
+++ b/Engine/Source/Runtime/StorageServerClient/Private/StorageServerPlatformFile.cpp
@@ -631,6 +631,52 @@ bool FStorageServerPlatformFile::DeleteDirectory(const TCHAR* Directory)
 	return LowerLevel->DeleteDirectory(Directory);
 }
 
+FString FStorageServerPlatformFile::ConvertToAbsolutePathForExternalAppForRead(const TCHAR* Filename)
+{
+#if PLATFORM_DESKTOP && (UE_GAME || UE_SERVER)
+	TStringBuilder<1024> Result;
+
+	// New code should not end up in here and should instead be written in such a
+	// way that data can be served from a (remote) server.
+
+	// Some data must exist in files on disk such that it can be accessed by external
+	// APIs. Any such data required by a title should have been written to Saved/Cooked
+	// at cook time. If a file prefix with UE's canonical ../../../ is requested we
+	// look inside Saved/Cooked. A read-only filesystem overlay if you will.
+
+	static FString* CookedDir = nullptr;
+	if (CookedDir == nullptr)
+	{
+		static FString Inner;
+		CookedDir = &Inner;
+
+		Result << *FPaths::ProjectDir();
+		Result << TEXT("Saved/Cooked/");
+		Result << FPlatformProperties::PlatformName();
+		Result << TEXT("/");
+		Inner = Result.ToString();
+	}
+	else
+	{
+		Result << *(*CookedDir);
+	}
+
+	const TCHAR* DotSlashSkip = Filename;
+	for (; *DotSlashSkip == '.' || *DotSlashSkip == '/'; ++DotSlashSkip);
+
+	if (PTRINT(DotSlashSkip - Filename) == 9) // 9 == ../../../
+	{
+		Result << DotSlashSkip;
+		if (LowerLevel->FileExists(Result.ToString()))
+		{
+			return FString(Result.GetData(), Result.Len());
+		}
+	}
+#endif
+
+	return LowerLevel->ConvertToAbsolutePathForExternalAppForRead(Filename);
+}
+
 bool FStorageServerPlatformFile::MakeStorageServerPath(const TCHAR* LocalFilenameOrDirectory, FStringBuilderBase& OutPath) const
 {
 	FStringView LocalEngineDirView(FPlatformMisc::EngineDir());
diff --git a/Engine/Source/Runtime/StorageServerClient/Private/StorageServerPlatformFile.h b/Engine/Source/Runtime/StorageServerClient/Private/StorageServerPlatformFile.h
--- a/Engine/Source/Runtime/StorageServerClient/Private/StorageServerPlatformFile.h
+++ b/Engine/Source/Runtime/StorageServerClient/Private/StorageServerPlatformFile.h
@@ -96,6 +96,7 @@ public:
 	virtual IFileHandle* OpenWrite(const TCHAR* Filename, bool bAppend = false, bool bAllowRead = false) override;
 	virtual bool CreateDirectory(const TCHAR* Directory) override;
 	virtual bool DeleteDirectory(const TCHAR* Directory) override;
+	virtual FString ConvertToAbsolutePathForExternalAppForRead(const TCHAR* Filename) override;
 	virtual bool SendMessageToServer(const TCHAR* Message, IPlatformFile::IFileServerMessageHandler* Handler) override;
 
 private:
