Moved FCookOnTheFlyPackageStore declaration to header file
#rnx

diff --git a/Engine/Source/Runtime/StorageServerClient/Private/CookOnTheFlyPackageStore.cpp b/Engine/Source/Runtime/StorageServerClient/Private/CookOnTheFlyPackageStore.cpp
--- a/Engine/Source/Runtime/StorageServerClient/Private/CookOnTheFlyPackageStore.cpp
+++ b/Engine/Source/Runtime/StorageServerClient/Private/CookOnTheFlyPackageStore.cpp
@@ -4,276 +4,223 @@
 
 #if WITH_COTF
 
-#include "CookOnTheFly.h"
 #include "HAL/PlatformTime.h"
 #include "Containers/ChunkedArray.h"
 #include "Serialization/MemoryReader.h"
 #include "CookOnTheFlyMessages.h"
 
-class FCookOnTheFlyPackageStore final
-	: public FPackageStoreBase
+FCookOnTheFlyPackageStore::FCookOnTheFlyPackageStore(UE::Cook::ICookOnTheFlyServerConnection& InCookOnTheFlyServerConnection)
+	: CookOnTheFlyServerConnection(InCookOnTheFlyServerConnection)
 {
-public:
-	struct FEntryInfo
-	{
-		EPackageStoreEntryStatus Status = EPackageStoreEntryStatus::None;
-		int32 EntryIndex = INDEX_NONE;
-	};
-
-	struct FPackageStats
-	{
-		TAtomic<uint32> Cooked{ 0 };
-		TAtomic<uint32> Failed{ 0 };
-	};
+	// Index zero is invalid
+	PackageEntries.Add();
+	CookOnTheFlyServerConnection.OnMessage().AddRaw(this, &FCookOnTheFlyPackageStore::OnCookOnTheFlyMessage);
 
-	FCookOnTheFlyPackageStore(UE::Cook::ICookOnTheFlyServerConnection& InCookOnTheFlyServerConnection)
-		: CookOnTheFlyServerConnection(InCookOnTheFlyServerConnection)
-	{
-		// Index zero is invalid
-		PackageEntries.Add();
-		CookOnTheFlyServerConnection.OnMessage().AddRaw(this, &FCookOnTheFlyPackageStore::OnCookOnTheFlyMessage);
+	using namespace UE::Cook;
+	using namespace UE::ZenCookOnTheFly::Messaging;
 
-		using namespace UE::Cook;
-		using namespace UE::ZenCookOnTheFly::Messaging;
+	FCookOnTheFlyRequest Request(ECookOnTheFlyMessage::GetCookedPackages);
+	FCookOnTheFlyResponse Response = CookOnTheFlyServerConnection.SendRequest(Request).Get();
 
-		FCookOnTheFlyRequest Request(ECookOnTheFlyMessage::GetCookedPackages);
-		FCookOnTheFlyResponse Response = CookOnTheFlyServerConnection.SendRequest(Request).Get();
+	if (Response.IsOk())
+	{
+		FGetCookedPackagesResponse GetCookedPackagesResponse = Response.GetBodyAs<FGetCookedPackagesResponse>();
+		FPackageStoreData& PackageStoreData = GetCookedPackagesResponse.PackageStoreData;
 
-		if (Response.IsOk())
-		{
-			FGetCookedPackagesResponse GetCookedPackagesResponse = Response.GetBodyAs<FGetCookedPackagesResponse>();
-			FPackageStoreData& PackageStoreData = GetCookedPackagesResponse.PackageStoreData;
+		UE_LOG(LogCookOnTheFly, Log, TEXT("Got '%d' cooked and '%d' failed packages from server"),
+			PackageStoreData.CookedPackages.Num(), PackageStoreData.FailedPackages.Num());
 
-			UE_LOG(LogCookOnTheFly, Log, TEXT("Got '%d' cooked and '%d' failed packages from server"),
-				PackageStoreData.CookedPackages.Num(), PackageStoreData.FailedPackages.Num());
+		AddPackages(MoveTemp(PackageStoreData.CookedPackages), MoveTemp(PackageStoreData.FailedPackages));
 
-			AddPackages(MoveTemp(PackageStoreData.CookedPackages), MoveTemp(PackageStoreData.FailedPackages));
-
-			LastServerActivtyTime = FPlatformTime::Seconds();
-		}
-		else
-		{
-			UE_LOG(LogCookOnTheFly, Warning, TEXT("Failed to send '%s' request"), LexToString(Request.GetHeader().MessageType));
-		}
+		LastServerActivtyTime = FPlatformTime::Seconds();
 	}
-
-	virtual void Initialize() override
+	else
 	{
-
+		UE_LOG(LogCookOnTheFly, Warning, TEXT("Failed to send '%s' request"), LexToString(Request.GetHeader().MessageType));
 	}
+}
 
-	virtual void Lock() override
-	{
-	}
+bool FCookOnTheFlyPackageStore::DoesPackageExist(FPackageId PackageId)
+{
+	FScopeLock _(&CriticalSection);
+	FEntryInfo EntryInfo = PackageIdToEntryInfo.FindRef(PackageId);
+	return EntryInfo.Status != EPackageStoreEntryStatus::Missing;
+}
 
-	virtual void Unlock() override
-	{
-	}
+EPackageStoreEntryStatus FCookOnTheFlyPackageStore::GetPackageStoreEntry(FPackageId PackageId, FPackageStoreEntry& OutPackageStoreEntry)
+{
+	using namespace UE::Cook;
+	using namespace UE::ZenCookOnTheFly::Messaging;
 
-	virtual bool DoesPackageExist(FPackageId PackageId) override
 	{
 		FScopeLock _(&CriticalSection);
-		FEntryInfo EntryInfo = PackageIdToEntryInfo.FindRef(PackageId);
-		return EntryInfo.Status != EPackageStoreEntryStatus::Missing;
-	}
-
-	virtual EPackageStoreEntryStatus GetPackageStoreEntry(FPackageId PackageId, FPackageStoreEntry& OutPackageStoreEntry) override
-	{
-		using namespace UE::Cook;
-		using namespace UE::ZenCookOnTheFly::Messaging;
+		FEntryInfo& EntryInfo = PackageIdToEntryInfo.FindOrAdd(PackageId, FEntryInfo());
 
+		if (EntryInfo.Status == EPackageStoreEntryStatus::Ok)
 		{
-			FScopeLock _(&CriticalSection);
-			FEntryInfo& EntryInfo = PackageIdToEntryInfo.FindOrAdd(PackageId, FEntryInfo());
-
-			if (EntryInfo.Status == EPackageStoreEntryStatus::Ok)
-			{
-				check(EntryInfo.EntryIndex != INDEX_NONE);
-				return CreatePackageStoreEntry(EntryInfo, OutPackageStoreEntry);
-			}
+			check(EntryInfo.EntryIndex != INDEX_NONE);
+			return CreatePackageStoreEntry(EntryInfo, OutPackageStoreEntry);
+		}
 			
-			if (EntryInfo.Status == EPackageStoreEntryStatus::Missing)
-			{
-				return EPackageStoreEntryStatus::Missing;
-			}
-
-			if (EntryInfo.Status == EPackageStoreEntryStatus::Pending)
-			{
-				CheckActivity();
-				return EPackageStoreEntryStatus::Pending;
-			}
-
-			// The package hasn't been requested, set the status to pending.
-			EntryInfo.Status = EPackageStoreEntryStatus::Pending;
+		if (EntryInfo.Status == EPackageStoreEntryStatus::Missing)
+		{
+			return EPackageStoreEntryStatus::Missing;
 		}
 
-		LastClientActivtyTime = FPlatformTime::Seconds();
-		UE_LOG(LogCookOnTheFly, Verbose, TEXT("Requesting package 0x%llX"), PackageId.ValueForDebugging());
-
-		FCookOnTheFlyRequest Request(ECookOnTheFlyMessage::CookPackage);
-		Request.SetBodyTo(FCookPackageRequest { PackageId });
-		FCookOnTheFlyResponse Response = CookOnTheFlyServerConnection.SendRequest(Request).Get();
-
-		if (!Response.IsOk())
+		if (EntryInfo.Status == EPackageStoreEntryStatus::Pending)
 		{
-			UE_LOG(LogCookOnTheFly, Warning, TEXT("Failed to send '%s' request"), LexToString(Request.GetHeader().MessageType));
-			FScopeLock _(&CriticalSection);
-			FEntryInfo& EntryInfo = PackageIdToEntryInfo.FindChecked(PackageId);
-			EntryInfo.Status = EPackageStoreEntryStatus::Missing;
-			return EPackageStoreEntryStatus::Missing;
+			CheckActivity();
+			return EPackageStoreEntryStatus::Pending;
 		}
 
-		FCookPackageResponse CookPackageResponse = Response.GetBodyAs<FCookPackageResponse>();
+		// The package hasn't been requested, set the status to pending.
+		EntryInfo.Status = EPackageStoreEntryStatus::Pending;
+	}
 
-		{
-			FScopeLock _(&CriticalSection);
+	LastClientActivtyTime = FPlatformTime::Seconds();
+	UE_LOG(LogCookOnTheFly, Verbose, TEXT("Requesting package 0x%llX"), PackageId.ValueForDebugging());
 
-			FEntryInfo& EntryInfo = PackageIdToEntryInfo.FindChecked(PackageId);
-			if (CookPackageResponse.Status == EPackageStoreEntryStatus::Missing)
-			{
-				EntryInfo.Status = EPackageStoreEntryStatus::Missing;
-				return EPackageStoreEntryStatus::Missing;
-			}
-			else
-			{
-				return CreatePackageStoreEntry(EntryInfo, OutPackageStoreEntry);
-			}
-		}
+	FCookOnTheFlyRequest Request(ECookOnTheFlyMessage::CookPackage);
+	Request.SetBodyTo(FCookPackageRequest { PackageId });
+	FCookOnTheFlyResponse Response = CookOnTheFlyServerConnection.SendRequest(Request).Get();
+
+	if (!Response.IsOk())
+	{
+		UE_LOG(LogCookOnTheFly, Warning, TEXT("Failed to send '%s' request"), LexToString(Request.GetHeader().MessageType));
+		FScopeLock _(&CriticalSection);
+		FEntryInfo& EntryInfo = PackageIdToEntryInfo.FindChecked(PackageId);
+		EntryInfo.Status = EPackageStoreEntryStatus::Missing;
+		return EPackageStoreEntryStatus::Missing;
 	}
 
-	EPackageStoreEntryStatus CreatePackageStoreEntry(const FEntryInfo& EntryInfo, FPackageStoreEntry& OutPackageStoreEntry)
+	FCookPackageResponse CookPackageResponse = Response.GetBodyAs<FCookPackageResponse>();
+
 	{
-		if (EntryInfo.Status == EPackageStoreEntryStatus::Ok)
+		FScopeLock _(&CriticalSection);
+
+		FEntryInfo& EntryInfo = PackageIdToEntryInfo.FindChecked(PackageId);
+		if (CookPackageResponse.Status == EPackageStoreEntryStatus::Missing)
 		{
-			const FPackageStoreEntryResource& Entry = PackageEntries[EntryInfo.EntryIndex];
-			OutPackageStoreEntry.ExportInfo = Entry.ExportInfo;
-			OutPackageStoreEntry.ImportedPackageIds = Entry.ImportedPackageIds;
-			return EPackageStoreEntryStatus::Ok;
+			EntryInfo.Status = EPackageStoreEntryStatus::Missing;
+			return EPackageStoreEntryStatus::Missing;
 		}
 		else
 		{
-			return EntryInfo.Status;
+			return CreatePackageStoreEntry(EntryInfo, OutPackageStoreEntry);
 		}
 	}
+}
 
-	virtual bool GetPackageRedirectInfo(FPackageId PackageId, FName& OutSourcePackageName, FPackageId& OutRedirectedToPackageId)  override
+EPackageStoreEntryStatus FCookOnTheFlyPackageStore::CreatePackageStoreEntry(const FEntryInfo& EntryInfo, FPackageStoreEntry& OutPackageStoreEntry)
+{
+	if (EntryInfo.Status == EPackageStoreEntryStatus::Ok)
 	{
-		return false;
+		const FPackageStoreEntryResource& Entry = PackageEntries[EntryInfo.EntryIndex];
+		OutPackageStoreEntry.ExportInfo = Entry.ExportInfo;
+		OutPackageStoreEntry.ImportedPackageIds = Entry.ImportedPackageIds;
+		return EPackageStoreEntryStatus::Ok;
 	}
-
-	void AddPackages(TArray<FPackageStoreEntryResource> Entries, TArray<FPackageId> FailedPackageIds)
+	else
 	{
-		FScopeLock _(&CriticalSection);
+		return EntryInfo.Status;
+	}
+}
+
+void FCookOnTheFlyPackageStore::AddPackages(TArray<FPackageStoreEntryResource> Entries, TArray<FPackageId> FailedPackageIds)
+{
+	FScopeLock _(&CriticalSection);
 		
-		for (FPackageId FailedPackageId : FailedPackageIds)
-		{
-			UE_LOG(LogCookOnTheFly, Verbose, TEXT("'0x%llX' [Failed]"), FailedPackageId.Value());
-			FEntryInfo& EntryInfo = PackageIdToEntryInfo.FindOrAdd(FailedPackageId, FEntryInfo());
-			EntryInfo.Status = EPackageStoreEntryStatus::Missing;
-			PackageStats.Failed++;
-		}
+	for (FPackageId FailedPackageId : FailedPackageIds)
+	{
+		UE_LOG(LogCookOnTheFly, Verbose, TEXT("'0x%llX' [Failed]"), FailedPackageId.Value());
+		FEntryInfo& EntryInfo = PackageIdToEntryInfo.FindOrAdd(FailedPackageId, FEntryInfo());
+		EntryInfo.Status = EPackageStoreEntryStatus::Missing;
+		PackageStats.Failed++;
+	}
 
-		const int32 NumPackageEntries = PackageEntries.Num();
+	const int32 NumPackageEntries = PackageEntries.Num();
 
-		for (FPackageStoreEntryResource& Entry : Entries)
-		{
-			const FPackageId PackageId = Entry.GetPackageId();
-			FEntryInfo& EntryInfo = PackageIdToEntryInfo.FindOrAdd(PackageId, FEntryInfo());
+	for (FPackageStoreEntryResource& Entry : Entries)
+	{
+		const FPackageId PackageId = Entry.GetPackageId();
+		FEntryInfo& EntryInfo = PackageIdToEntryInfo.FindOrAdd(PackageId, FEntryInfo());
 
-			if (EntryInfo.EntryIndex == INDEX_NONE)
-			{
-				EntryInfo.Status = EPackageStoreEntryStatus::Ok;
-				EntryInfo.EntryIndex = PackageEntries.Add();
-				PackageEntries[EntryInfo.EntryIndex] = MoveTemp(Entry);
-				PackageStats.Cooked++;
+		if (EntryInfo.EntryIndex == INDEX_NONE)
+		{
+			EntryInfo.Status = EPackageStoreEntryStatus::Ok;
+			EntryInfo.EntryIndex = PackageEntries.Add();
+			PackageEntries[EntryInfo.EntryIndex] = MoveTemp(Entry);
+			PackageStats.Cooked++;
 
-				UE_LOG(LogCookOnTheFly, Verbose, TEXT("'%s' [OK] (Cooked/Failed='%d/%d')"),
-					*PackageEntries[EntryInfo.EntryIndex].PackageName.ToString(), PackageStats.Cooked.Load(), PackageStats.Failed.Load());
-			}
+			UE_LOG(LogCookOnTheFly, Verbose, TEXT("'%s' [OK] (Cooked/Failed='%d/%d')"),
+				*PackageEntries[EntryInfo.EntryIndex].PackageName.ToString(), PackageStats.Cooked.Load(), PackageStats.Failed.Load());
 		}
 	}
+}
 
-	void OnCookOnTheFlyMessage(const UE::Cook::FCookOnTheFlyMessage& Message)
-	{
-		using namespace UE::ZenCookOnTheFly::Messaging;
+void FCookOnTheFlyPackageStore::OnCookOnTheFlyMessage(const UE::Cook::FCookOnTheFlyMessage& Message)
+{
+	using namespace UE::ZenCookOnTheFly::Messaging;
 
-		switch (Message.GetHeader().MessageType)
+	switch (Message.GetHeader().MessageType)
+	{
+		case UE::Cook::ECookOnTheFlyMessage::PackagesCooked:
 		{
-			case UE::Cook::ECookOnTheFlyMessage::PackagesCooked:
-			{
-				FPackagesCookedMessage PackagesCookedMessage = Message.GetBodyAs<FPackagesCookedMessage>();
-				FPackageStoreData& PackageStoreData = PackagesCookedMessage.PackageStoreData;
-
-				UE_LOG(LogCookOnTheFly, Verbose, TEXT("Received '%s' message, Cooked='%d', Failed='%d', Server total='%d/%d' (Cooked/Failed)"),
-					LexToString(Message.GetHeader().MessageType),
-					PackageStoreData.CookedPackages.Num(),
-					PackageStoreData.FailedPackages.Num(),
-					PackageStoreData.TotalCookedPackages,
-					PackageStoreData.TotalFailedPackages);
+			FPackagesCookedMessage PackagesCookedMessage = Message.GetBodyAs<FPackagesCookedMessage>();
+			FPackageStoreData& PackageStoreData = PackagesCookedMessage.PackageStoreData;
 
-				AddPackages(MoveTemp(PackageStoreData.CookedPackages), MoveTemp(PackageStoreData.FailedPackages));
+			UE_LOG(LogCookOnTheFly, Verbose, TEXT("Received '%s' message, Cooked='%d', Failed='%d', Server total='%d/%d' (Cooked/Failed)"),
+				LexToString(Message.GetHeader().MessageType),
+				PackageStoreData.CookedPackages.Num(),
+				PackageStoreData.FailedPackages.Num(),
+				PackageStoreData.TotalCookedPackages,
+				PackageStoreData.TotalFailedPackages);
 
-				UE_CLOG(PackageStoreData.TotalCookedPackages != PackageStats.Cooked.Load() || PackageStoreData.TotalFailedPackages != PackageStats.Failed.Load(),
-					LogCookOnTheFly, Warning, TEXT("Client/Server package mismatch, Cooked='%d/%d', Failed='%d/%d' (Client/Server)"),
-					PackageStats.Cooked.Load(), PackageStoreData.TotalCookedPackages, PackageStats.Failed.Load(), PackageStoreData.TotalFailedPackages);
+			AddPackages(MoveTemp(PackageStoreData.CookedPackages), MoveTemp(PackageStoreData.FailedPackages));
 
-				PendingEntriesAdded.Broadcast();
+			UE_CLOG(PackageStoreData.TotalCookedPackages != PackageStats.Cooked.Load() || PackageStoreData.TotalFailedPackages != PackageStats.Failed.Load(),
+				LogCookOnTheFly, Warning, TEXT("Client/Server package mismatch, Cooked='%d/%d', Failed='%d/%d' (Client/Server)"),
+				PackageStats.Cooked.Load(), PackageStoreData.TotalCookedPackages, PackageStats.Failed.Load(), PackageStoreData.TotalFailedPackages);
 
-				break;
-			}
+			PendingEntriesAdded.Broadcast();
 
-			default:
-				break;
+			break;
 		}
 
-		LastServerActivtyTime = FPlatformTime::Seconds();
+		default:
+			break;
 	}
 
-	void CheckActivity()
-	{
-		const double TimeSinceLastClientActivity = FPlatformTime::Seconds() - LastClientActivtyTime;
-		const double TimeSinceLastServerActivity = FPlatformTime::Seconds() - LastServerActivtyTime;
-		const double TimeSinceLastWarning = FPlatformTime::Seconds() - LastWarningTime;
+	LastServerActivtyTime = FPlatformTime::Seconds();
+}
 
-		if (TimeSinceLastClientActivity > MaxInactivityTime &&
-			TimeSinceLastServerActivity > MaxInactivityTime &&
-			TimeSinceLastWarning > TimeBetweenWarning)
-		{
-			LastWarningTime = FPlatformTime::Seconds();
+void FCookOnTheFlyPackageStore::CheckActivity()
+{
+	const double TimeSinceLastClientActivity = FPlatformTime::Seconds() - LastClientActivtyTime;
+	const double TimeSinceLastServerActivity = FPlatformTime::Seconds() - LastServerActivtyTime;
+	const double TimeSinceLastWarning = FPlatformTime::Seconds() - LastWarningTime;
+
+	if (TimeSinceLastClientActivity > MaxInactivityTime &&
+		TimeSinceLastServerActivity > MaxInactivityTime &&
+		TimeSinceLastWarning > TimeBetweenWarning)
+	{
+		LastWarningTime = FPlatformTime::Seconds();
 
-			UE_LOG(LogCookOnTheFly, Log, TEXT("No server response in '%.2lf' seconds"), TimeSinceLastServerActivity);
+		UE_LOG(LogCookOnTheFly, Log, TEXT("No server response in '%.2lf' seconds"), TimeSinceLastServerActivity);
 
-			UE_LOG(LogCookOnTheFly, Log, TEXT("=== Pending Packages ==="));
+		UE_LOG(LogCookOnTheFly, Log, TEXT("=== Pending Packages ==="));
+		{
+			FScopeLock _(&CriticalSection);
+			for (const auto& KeyValue : PackageIdToEntryInfo)
 			{
-				FScopeLock _(&CriticalSection);
-				for (const auto& KeyValue : PackageIdToEntryInfo)
+				const FEntryInfo& EntryInfo = KeyValue.Value;
+				if (EntryInfo.Status == EPackageStoreEntryStatus::Pending)
 				{
-					const FEntryInfo& EntryInfo = KeyValue.Value;
-					if (EntryInfo.Status == EPackageStoreEntryStatus::Pending)
-					{
-						UE_LOG(LogCookOnTheFly, Log, TEXT("0x%llX"), KeyValue.Key.ValueForDebugging());
-					}
+					UE_LOG(LogCookOnTheFly, Log, TEXT("0x%llX"), KeyValue.Key.ValueForDebugging());
 				}
 			}
 		}
 	}
-
-	UE::Cook::ICookOnTheFlyServerConnection& CookOnTheFlyServerConnection;
-	FCriticalSection CriticalSection;
-	TMap<FPackageId, FEntryInfo> PackageIdToEntryInfo;
-	TChunkedArray<FPackageStoreEntryResource> PackageEntries;
-	FPackageStats PackageStats;
-
-	const double MaxInactivityTime = 20;
-	const double TimeBetweenWarning = 10;
-	double LastClientActivtyTime = 0;
-	double LastServerActivtyTime = 0;
-	double LastWarningTime = 0;
-};
-
-TSharedPtr<IPackageStore> MakeCookOnTheFlyPackageStore(UE::Cook::ICookOnTheFlyServerConnection& CookOnTheFlyServerConnection)
-{
-	return MakeShared<FCookOnTheFlyPackageStore>(CookOnTheFlyServerConnection);
 }
 
 #endif // WITH_COTF
diff --git a/Engine/Source/Runtime/StorageServerClient/Private/CookOnTheFlyPackageStore.h b/Engine/Source/Runtime/StorageServerClient/Private/CookOnTheFlyPackageStore.h
--- a/Engine/Source/Runtime/StorageServerClient/Private/CookOnTheFlyPackageStore.h
+++ b/Engine/Source/Runtime/StorageServerClient/Private/CookOnTheFlyPackageStore.h
@@ -7,14 +7,68 @@
 #if WITH_COTF
 
 #include "IO/PackageStore.h"
-
-class FIoDispatcher;
+#include "CookOnTheFly.h"
 
 namespace UE { namespace Cook
 {
 	class ICookOnTheFlyServerConnection;
 }}
 
-TSharedPtr<IPackageStore> MakeCookOnTheFlyPackageStore(UE::Cook::ICookOnTheFlyServerConnection& CookOnTheFlyServerConnection);
+class FCookOnTheFlyPackageStore final
+	: public FPackageStoreBase
+{
+public:
+	struct FEntryInfo
+	{
+		EPackageStoreEntryStatus Status = EPackageStoreEntryStatus::None;
+		int32 EntryIndex = INDEX_NONE;
+	};
+
+	struct FPackageStats
+	{
+		TAtomic<uint32> Cooked{ 0 };
+		TAtomic<uint32> Failed{ 0 };
+	};
+
+	FCookOnTheFlyPackageStore(UE::Cook::ICookOnTheFlyServerConnection& InCookOnTheFlyServerConnection);
+
+	virtual void Initialize() override
+	{
+	}
+
+	virtual void Lock() override
+	{
+	}
+
+	virtual void Unlock() override
+	{
+	}
+
+	virtual bool DoesPackageExist(FPackageId PackageId) override;
+	virtual EPackageStoreEntryStatus GetPackageStoreEntry(FPackageId PackageId, FPackageStoreEntry& OutPackageStoreEntry) override;
+	
+	virtual bool GetPackageRedirectInfo(FPackageId PackageId, FName& OutSourcePackageName, FPackageId& OutRedirectedToPackageId) override
+	{
+		return false;
+	}
+	
+private:
+	EPackageStoreEntryStatus CreatePackageStoreEntry(const FEntryInfo& EntryInfo, FPackageStoreEntry& OutPackageStoreEntry);
+	void AddPackages(TArray<FPackageStoreEntryResource> Entries, TArray<FPackageId> FailedPackageIds);
+	void OnCookOnTheFlyMessage(const UE::Cook::FCookOnTheFlyMessage& Message);
+	void CheckActivity();
+
+	UE::Cook::ICookOnTheFlyServerConnection& CookOnTheFlyServerConnection;
+	FCriticalSection CriticalSection;
+	TMap<FPackageId, FEntryInfo> PackageIdToEntryInfo;
+	TChunkedArray<FPackageStoreEntryResource> PackageEntries;
+	FPackageStats PackageStats;
+
+	const double MaxInactivityTime = 20;
+	const double TimeBetweenWarning = 10;
+	double LastClientActivtyTime = 0;
+	double LastServerActivtyTime = 0;
+	double LastWarningTime = 0;
+};
 
 #endif // WITH_COTF
diff --git a/Engine/Source/Runtime/StorageServerClient/Private/StorageServerPlatformFile.cpp b/Engine/Source/Runtime/StorageServerClient/Private/StorageServerPlatformFile.cpp
--- a/Engine/Source/Runtime/StorageServerClient/Private/StorageServerPlatformFile.cpp
+++ b/Engine/Source/Runtime/StorageServerClient/Private/StorageServerPlatformFile.cpp
@@ -319,7 +319,7 @@ void FStorageServerPlatformFile::InitializeAfterProjectFilePath()
 #if WITH_COTF
 				if (CookOnTheFlyServerConnection)
 				{
-					return MakeCookOnTheFlyPackageStore(*CookOnTheFlyServerConnection.Get());
+					return MakeShared<FCookOnTheFlyPackageStore>(*CookOnTheFlyServerConnection.Get());
 				}
 				else
 #endif
