ZenCOTF: Fixed shutdown crash caused by server connection outliving the socket system
#rnx

diff --git a/Engine/Source/Runtime/StorageServerClient/Private/StorageServerPlatformFile.cpp b/Engine/Source/Runtime/StorageServerClient/Private/StorageServerPlatformFile.cpp
--- a/Engine/Source/Runtime/StorageServerClient/Private/StorageServerPlatformFile.cpp
+++ b/Engine/Source/Runtime/StorageServerClient/Private/StorageServerPlatformFile.cpp
@@ -244,10 +244,22 @@ public:
 
 FStorageServerPlatformFile::FStorageServerPlatformFile()
 {
+	FCoreDelegates::OnEnginePreExit.AddRaw(this, &FStorageServerPlatformFile::OnEnginePreExit);
 }
 
 FStorageServerPlatformFile::~FStorageServerPlatformFile()
 {
+	FCoreDelegates::OnEnginePreExit.RemoveAll(this);
+}
+
+void FStorageServerPlatformFile::OnEnginePreExit()
+{
+#if WITH_COTF
+	if (CookOnTheFlyServerConnection)
+	{
+		CookOnTheFlyServerConnection->Disconnect();
+	}
+#endif
 }
 
 bool FStorageServerPlatformFile::ShouldBeUsed(IPlatformFile* Inner, const TCHAR* CmdLine) const
diff --git a/Engine/Source/Runtime/StorageServerClient/Private/StorageServerPlatformFile.h b/Engine/Source/Runtime/StorageServerClient/Private/StorageServerPlatformFile.h
--- a/Engine/Source/Runtime/StorageServerClient/Private/StorageServerPlatformFile.h
+++ b/Engine/Source/Runtime/StorageServerClient/Private/StorageServerPlatformFile.h
@@ -102,6 +102,7 @@ public:
 private:
 	friend class FStorageServerFileHandle;
 
+	void OnEnginePreExit();
 	bool MakeStorageServerPath(const TCHAR* LocalFilenameOrDirectory, FStringBuilderBase& OutPath) const;
 	bool MakeLocalPath(const TCHAR* ServerFilenameOrDirectory, FStringBuilderBase& OutPath) const;
 	IFileHandle* InternalOpenFile(const FIoChunkId& FileChunkId, const TCHAR* LocalFilename);
