Fix static analysis warning (declaring virtual methods in final class is pointless)

diff --git a/Engine/Source/Runtime/StorageServerClient/Private/CookOnTheFlyPackageStore.h b/Engine/Source/Runtime/StorageServerClient/Private/CookOnTheFlyPackageStore.h
--- a/Engine/Source/Runtime/StorageServerClient/Private/CookOnTheFlyPackageStore.h
+++ b/Engine/Source/Runtime/StorageServerClient/Private/CookOnTheFlyPackageStore.h
@@ -40,7 +40,7 @@ public:
 	{
 	}
 
-	virtual bool DoesPackageExist(FPackageId PackageId);
+	bool DoesPackageExist(FPackageId PackageId);
 	virtual EPackageStoreEntryStatus GetPackageStoreEntry(FPackageId PackageId, FPackageStoreEntry& OutPackageStoreEntry) override;
 	
 	virtual bool GetPackageRedirectInfo(FPackageId PackageId, FName& OutSourcePackageName, FPackageId& OutRedirectedToPackageId) override
