Zen: Fixed error converting non null terminated error string from ANSI to TCHAR

#rnx

diff --git a/Engine/Source/Runtime/StorageServerClient/Private/StorageServerConnection.cpp b/Engine/Source/Runtime/StorageServerClient/Private/StorageServerConnection.cpp
--- a/Engine/Source/Runtime/StorageServerClient/Private/StorageServerConnection.cpp
+++ b/Engine/Source/Runtime/StorageServerClient/Private/StorageServerConnection.cpp
@@ -191,10 +191,20 @@
 	if (!bIsOk && ContentLength)
 	{
 		TArray<uint8> ErrorBuffer;
-		ErrorBuffer.SetNumUninitialized(ContentLength);
-		int32 BytesRead;
+		ErrorBuffer.SetNumUninitialized(ContentLength + 1);
+		int32 BytesRead = 0;
 		InSocket.Recv(ErrorBuffer.GetData(), ContentLength, BytesRead, ESocketReceiveFlags::WaitAll);
-		ErrorMessage = FString(ContentLength, ANSI_TO_TCHAR(reinterpret_cast<ANSICHAR*>(ErrorBuffer.GetData())));
+
+		if (BytesRead > 0)
+		{
+			ErrorBuffer[BytesRead] = '\0';
+			ErrorMessage = FString(BytesRead, ANSI_TO_TCHAR(reinterpret_cast<ANSICHAR*>(ErrorBuffer.GetData())));
+		}
+		else
+		{
+			ErrorMessage = TEXT("Unknown error");
+		}
+
 		ContentLength = 0;
 	}
 	if (ContentLength == 0)
