StorageServerClient: Drive-by cleanup of includes

#rnx

diff --git a/Engine/Source/Runtime/StorageServerClient/Private/StorageServerConnection.cpp b/Engine/Source/Runtime/StorageServerClient/Private/StorageServerConnection.cpp
--- a/Engine/Source/Runtime/StorageServerClient/Private/StorageServerConnection.cpp
+++ b/Engine/Source/Runtime/StorageServerClient/Private/StorageServerConnection.cpp
@@ -1,21 +1,19 @@
 // Copyright Epic Games, Inc. All Rights Reserved.
 
 #include "StorageServerConnection.h"
+
+#include "IO/IoDispatcher.h"
 #include "IPAddress.h"
-#include "SocketSubsystem.h"
-#include "Sockets.h"
 #include "Misc/App.h"
-#include "Misc/Paths.h"
-#include "Misc/SecureHash.h"
-#include "IO/IoDispatcher.h"
-#include "Misc/StringBuilder.h"
-#include "Misc/ScopeExit.h"
 #include "Misc/ScopeLock.h"
+#include "Misc/StringBuilder.h"
 #include "Serialization/CompactBinary.h"
 #include "Serialization/CompactBinarySerialization.h"
+#include "Memory/MemoryView.h"
 #include "Memory/SharedBuffer.h"
 #include "ProfilingDebugging/CountersTrace.h"
-#include "HAL/FileManager.h"
+#include "SocketSubsystem.h"
+#include "Sockets.h"
 
 #if !UE_BUILD_SHIPPING
 
@@ -418,6 +416,12 @@ int64 FStorageServerResponse::SerializeChunkTo(FMutableMemoryView Memory, uint64
 	return 0;
 }
 
+FCbObject FStorageServerResponse::GetResponseObject()
+{
+	FCbField Payload = LoadCompactBinary(*this);
+	return Payload.AsObject();
+}
+
 FStorageServerConnection::FStorageServerConnection()
 	: SocketSubsystem(*ISocketSubsystem::Get())
 {
diff --git a/Engine/Source/Runtime/StorageServerClient/Public/StorageServerConnection.h b/Engine/Source/Runtime/StorageServerClient/Public/StorageServerConnection.h
--- a/Engine/Source/Runtime/StorageServerClient/Public/StorageServerConnection.h
+++ b/Engine/Source/Runtime/StorageServerClient/Public/StorageServerConnection.h
@@ -2,21 +2,22 @@
 
 #pragma once
 
-#include "CoreMinimal.h"
+#include "CoreTypes.h"
 #include "Compression/CompressedBuffer.h"
-#include "IO/IoContainerId.h"
+#include "Containers/UnrealString.h"
 #include "Misc/StringBuilder.h"
-#include "Serialization/CompactBinarySerialization.h"
-#include "Memory/MemoryView.h"
+#include "Memory/MemoryFwd.h"
+#include "Templates/SharedPointer.h"
 
 #if !UE_BUILD_SHIPPING
 
-class ISocketSubsystem;
-class FInternetAddr;
+class FIoBuffer;
 class FIoChunkId;
+class FInternetAddr;
 class FSocket;
+class FStorageServerChunkBatchRequest;
 class FStorageServerConnection;
-class FIoBuffer;
+class ISocketSubsystem;
 
 enum class EStorageServerContentType : uint8
 {
@@ -72,7 +73,7 @@ class FStorageServerRequest
 {
 	EStorageServerContentType AcceptContentType() const;
 protected:
-	friend class FStorageServerConnection;
+	friend FStorageServerConnection;
 
 	FStorageServerRequest(
 		FAnsiStringView Verb,
@@ -137,17 +138,13 @@ public:
 
 	int64 SerializeChunkTo(FMutableMemoryView Memory, uint64 RawOffset = 0);
 
-	FCbObject GetResponseObject()
-	{
-		FCbField Payload = LoadCompactBinary(*this);
-		return Payload.AsObject();
-	}
+	FCbObject GetResponseObject();
 
 private:
-	friend class FStorageServerConnection;
-	friend class FStorageServerChunkBatchRequest;
+	friend FStorageServerConnection;
+	friend FStorageServerChunkBatchRequest;
 
-	FStorageServerResponse(class FStorageServerConnection& Owner, FSocket& Socket);
+	FStorageServerResponse(FStorageServerConnection& Owner, FSocket& Socket);
 	void ReleaseSocket(bool bKeepAlive);
 
 	FStorageServerConnection& Owner;
@@ -168,7 +165,7 @@ public:
 	STORAGESERVERCLIENT_API bool Issue(TFunctionRef<void(uint32 ChunkCount, uint32* ChunkIndices, uint64* ChunkSizes, FStorageServerResponse& ChunkDataStream)> OnResponse);
 
 private:
-	friend class FStorageServerConnection;
+	friend FStorageServerConnection;
 
 	FStorageServerChunkBatchRequest(FStorageServerConnection& Owner, FAnsiStringView Resource, FAnsiStringView Hostname);
 
@@ -192,9 +189,9 @@ public:
 	STORAGESERVERCLIENT_API FString GetHostAddr() const;
 
 private:
-	friend class FStorageServerRequest;
-	friend class FStorageServerResponse;
-	friend class FStorageServerChunkBatchRequest;
+	friend FStorageServerRequest;
+	friend FStorageServerResponse;
+	friend FStorageServerChunkBatchRequest;
 
 	int32 HandshakeRequest(TArrayView<const TSharedPtr<FInternetAddr>> HostAddresses);
 	FSocket* AcquireSocketFromPool();
