Fix for reading partial compressed buffers from zenserver project store.

diff --git a/Engine/Source/Runtime/StorageServerClient/Private/StorageServerConnection.cpp b/Engine/Source/Runtime/StorageServerClient/Private/StorageServerConnection.cpp
--- a/Engine/Source/Runtime/StorageServerClient/Private/StorageServerConnection.cpp
+++ b/Engine/Source/Runtime/StorageServerClient/Private/StorageServerConnection.cpp
@@ -401,8 +401,8 @@ int64 FStorageServerResponse::SerializeChunkTo(FMutableMemoryView Memory, uint64
 		
 		if (FCompressedBuffer Compressed = FCompressedBuffer::FromCompressed(FSharedBuffer::MakeView(CompressedBuffer.GetData(), ContentLength)))
 		{
-			FMutableMemoryView Dst = Memory.Left(FMath::Min(Memory.GetSize(), Compressed.GetRawSize()));
 			const uint64 CompressedOffset = GetCompressedOffset(Compressed, RawOffset);
+			FMutableMemoryView Dst = Memory.Left(FMath::Min(Memory.GetSize(), Compressed.GetRawSize() - CompressedOffset));
 			if (FCompressedBufferReader(Compressed).TryDecompressTo(Dst, CompressedOffset))
 			{
 				return Dst.GetSize();
