ZenIoDispatcher: Complete and delete batches when shutting down

diff --git a/Engine/Source/Runtime/StorageServerClient/Private/StorageServerIoDispatcherBackend.cpp b/Engine/Source/Runtime/StorageServerClient/Private/StorageServerIoDispatcherBackend.cpp
--- a/Engine/Source/Runtime/StorageServerClient/Private/StorageServerIoDispatcherBackend.cpp
+++ b/Engine/Source/Runtime/StorageServerClient/Private/StorageServerIoDispatcherBackend.cpp
@@ -32,11 +32,20 @@ FStorageServerIoDispatcherBackend::FStorageServerIoDispatcherBackend(FStorageSer
 
 FStorageServerIoDispatcherBackend::~FStorageServerIoDispatcherBackend()
 {
-	delete Thread;
+	Shutdown();
 	FPlatformProcess::ReturnSynchEventToPool(NewRequestEvent);
 	FPlatformProcess::ReturnSynchEventToPool(BatchCompletedEvent);
 }
 
+void FStorageServerIoDispatcherBackend::Shutdown()
+{
+	if (Thread)
+	{
+		delete Thread;
+		Thread = nullptr;
+	}
+}
+
 void FStorageServerIoDispatcherBackend::Initialize(TSharedRef<const FIoDispatcherBackendContext> InContext)
 {
 	BackendContext = InContext;
@@ -51,7 +60,9 @@ void FStorageServerIoDispatcherBackend::Stop()
 
 uint32 FStorageServerIoDispatcherBackend::Run()
 {
-	for (int32 BatchIndex = 0; BatchIndex < GStorageServerIoDispatcherMaxActiveBatchCount; ++BatchIndex)
+	LLM_SCOPE(ELLMTag::FileSystem);
+	const int32 BatchCount = GStorageServerIoDispatcherMaxActiveBatchCount;
+	for (int32 BatchIndex = 0; BatchIndex < BatchCount; ++BatchIndex)
 	{
 		FBatch* Batch = new FBatch(*this, MakeUnique<FStorageServerSerializationContext>());
 		Batch->Next = FirstAvailableBatch;
@@ -108,6 +119,21 @@ uint32 FStorageServerIoDispatcherBackend::Run()
 		NewRequestEvent->Wait();
 	}
 
+	for (int32 BatchIndex = 0; BatchIndex < BatchCount; ++BatchIndex)
+	{
+		if (!FirstAvailableBatch)
+		{
+			if (!WaitForBatchToComplete(10000))
+			{
+				UE_LOG(LogIoDispatcher, Warning, TEXT("Outstanding requests when shutting down storage server backend"));
+				return 0;
+			}
+		}
+		check(FirstAvailableBatch);
+		FBatch* Batch = FirstAvailableBatch;
+		FirstAvailableBatch = FirstAvailableBatch->Next;
+		delete Batch;
+	}
 	return 0;
 }
 
@@ -185,12 +211,15 @@ void FStorageServerIoDispatcherBackend::SubmitBatch(FBatch* Batch)
 	++SubmittedBatchesCount;
 }
 
-void FStorageServerIoDispatcherBackend::WaitForBatchToComplete()
+bool FStorageServerIoDispatcherBackend::WaitForBatchToComplete(uint32 WaitTime)
 {
 	bool bAtLeastOneCompleted = false;
 	while (!bAtLeastOneCompleted)
 	{
-		BatchCompletedEvent->Wait();
+		if (!BatchCompletedEvent->Wait(WaitTime))
+		{
+			return false;
+		}
 
 		FBatch* LocalCompletedBatches;
 		{
@@ -209,6 +238,7 @@ void FStorageServerIoDispatcherBackend::WaitForBatchToComplete()
 			bAtLeastOneCompleted = true;
 		}
 	}
+	return true;
 }
 
 void FStorageServerIoDispatcherBackend::OnBatchCompleted(FBatch* Batch)
@@ -239,6 +269,7 @@ void FStorageServerIoDispatcherBackend::OnBatchCompleted(FBatch* Batch)
 
 void FStorageServerIoDispatcherBackend::FBatch::DoThreadedWork()
 {
+	LLM_SCOPE(ELLMTag::FileSystem);
 	TRACE_CPUPROFILER_EVENT_SCOPE(StorageServerIoDispatcherProcessBatch);
 	FIoRequestImpl* Request = RequestsHead;
 #if 0
@@ -270,7 +301,6 @@ void FStorageServerIoDispatcherBackend::FBatch::DoThreadedWork()
 				}
 				else
 				{
-					LLM_SCOPE(ELLMTag::FileSystem);
 					TRACE_CPUPROFILER_EVENT_SCOPE(AllocMemoryForRequest);
 					Request->IoBuffer = FIoBuffer(RequestSize);
 				}
diff --git a/Engine/Source/Runtime/StorageServerClient/Private/StorageServerIoDispatcherBackend.h b/Engine/Source/Runtime/StorageServerClient/Private/StorageServerIoDispatcherBackend.h
--- a/Engine/Source/Runtime/StorageServerClient/Private/StorageServerIoDispatcherBackend.h
+++ b/Engine/Source/Runtime/StorageServerClient/Private/StorageServerIoDispatcherBackend.h
@@ -20,6 +20,7 @@ public:
 	~FStorageServerIoDispatcherBackend();
 
 	void Initialize(TSharedRef<const FIoDispatcherBackendContext> Context) override;
+	void Shutdown() override;
 	bool Resolve(FIoRequestImpl* Request) override;
 	void CancelIoRequest(FIoRequestImpl* Request) override {};
 	void UpdatePriorityForIoRequest(FIoRequestImpl* Request) override {};
@@ -85,7 +86,7 @@ private:
 
 	void SubmitBatch(FBatch* Batch);
 	void OnBatchCompleted(FBatch* Batch);
-	void WaitForBatchToComplete();
+	bool WaitForBatchToComplete(uint32 WaitTime = MAX_uint32);
 
 	FStorageServerConnection& Connection;
 	TSharedPtr<const FIoDispatcherBackendContext> BackendContext;
