Some more non unity/pch compile errors for linux target

diff --git a/Engine/Source/Runtime/StorageServerClient/Private/StorageServerIoDispatcherBackend.cpp b/Engine/Source/Runtime/StorageServerClient/Private/StorageServerIoDispatcherBackend.cpp
--- a/Engine/Source/Runtime/StorageServerClient/Private/StorageServerIoDispatcherBackend.cpp
+++ b/Engine/Source/Runtime/StorageServerClient/Private/StorageServerIoDispatcherBackend.cpp
@@ -267,6 +267,13 @@ void FStorageServerIoDispatcherBackend::OnBatchCompleted(FBatch* Batch)
 	BatchCompletedEvent->Trigger();
 }
 
+FStorageServerIoDispatcherBackend::FBatch::FBatch(FStorageServerIoDispatcherBackend& InOwner, TUniquePtr<FStorageServerSerializationContext> InSerializationContext)
+	: Owner(InOwner)
+	, SerializationContext(MoveTemp(InSerializationContext))
+{
+
+}
+
 void FStorageServerIoDispatcherBackend::FBatch::DoThreadedWork()
 {
 	LLM_SCOPE(ELLMTag::FileSystem);
diff --git a/Engine/Source/Runtime/StorageServerClient/Private/StorageServerIoDispatcherBackend.h b/Engine/Source/Runtime/StorageServerClient/Private/StorageServerIoDispatcherBackend.h
--- a/Engine/Source/Runtime/StorageServerClient/Private/StorageServerIoDispatcherBackend.h
+++ b/Engine/Source/Runtime/StorageServerClient/Private/StorageServerIoDispatcherBackend.h
@@ -66,12 +66,7 @@ private:
 	struct FBatch
 		: public IQueuedWork
 	{
-		FBatch(FStorageServerIoDispatcherBackend& InOwner, TUniquePtr<FStorageServerSerializationContext> InSerializationContext)
-			: Owner(InOwner)
-			, SerializationContext(MoveTemp(InSerializationContext))
-		{
-
-		}
+		FBatch(FStorageServerIoDispatcherBackend& InOwner, TUniquePtr<FStorageServerSerializationContext> InSerializationContext);
 
 		virtual void DoThreadedWork() override;
 		virtual void Abandon() override {};
