Fix bug with handling of IPV6 addresses in StorageServerConnection.

diff --git a/Engine/Source/Runtime/StorageServerClient/Private/StorageServerConnection.cpp b/Engine/Source/Runtime/StorageServerClient/Private/StorageServerConnection.cpp
--- a/Engine/Source/Runtime/StorageServerClient/Private/StorageServerConnection.cpp
+++ b/Engine/Source/Runtime/StorageServerClient/Private/StorageServerConnection.cpp
@@ -24,15 +24,23 @@ TRACE_DECLARE_MEMORY_COUNTER(ZenHttpClientSerializedBytes, TEXT("ZenClient/Seria
 
 static TArray<TSharedPtr<FInternetAddr>> GetAddressFromString(ISocketSubsystem& SocketSubsystem, TArrayView<const FString> HostAddresses, const int32 Port)
 {
-	TArray<TSharedPtr<FInternetAddr>> InterntAddresses;
+	TArray<TSharedPtr<FInternetAddr>> InternetAddresses;
+	FString ModifiedHostAddr;
 
 	for (const FString& HostAddr : HostAddresses)
 	{
-		TSharedPtr<FInternetAddr> Addr = SocketSubsystem.GetAddressFromString(HostAddr);
+		// Numeric IPV6 addresses can be enclosed in brackets, and must have the brackets stripped before calling GetAddressFromString
+		const FString* EffectiveHostAddr = &HostAddr;
+		if (!HostAddr.IsEmpty() && HostAddr[0] == TEXT('[') && HostAddr[HostAddr.Len() - 1] == TEXT(']'))
+		{
+			ModifiedHostAddr = FStringView(HostAddr).Mid(1,HostAddr.Len() - 2);
+			EffectiveHostAddr = &ModifiedHostAddr;
+		}
+		TSharedPtr<FInternetAddr> Addr = SocketSubsystem.GetAddressFromString(*EffectiveHostAddr);
 
 		if (!Addr.IsValid() || !Addr->IsValid())
 		{
-			FAddressInfoResult GAIRequest = SocketSubsystem.GetAddressInfo(*HostAddr, nullptr, EAddressInfoFlags::Default, NAME_None);
+			FAddressInfoResult GAIRequest = SocketSubsystem.GetAddressInfo(**EffectiveHostAddr, nullptr, EAddressInfoFlags::Default, NAME_None);
 			if (GAIRequest.ReturnCode == SE_NO_ERROR && GAIRequest.Results.Num() > 0)
 			{
 				Addr = GAIRequest.Results[0].Address;
@@ -42,11 +50,11 @@ static TArray<TSharedPtr<FInternetAddr>> GetAddressFromString(ISocketSubsystem&
 		if (Addr.IsValid() && Addr->IsValid())
 		{
 			Addr->SetPort(Port);
-			InterntAddresses.Emplace(MoveTemp(Addr));
+			InternetAddresses.Emplace(MoveTemp(Addr));
 		}
 	}
 
-	return InterntAddresses;
+	return InternetAddresses;
 }
 
 static uint64 GetCompressedOffset(const FCompressedBuffer& Buffer, uint64 RawOffset)
