Simplify detection of .projectstore file by using the relative content root

diff --git a/Engine/Source/Runtime/StorageServerClient/Private/StorageServerPlatformFile.cpp b/Engine/Source/Runtime/StorageServerClient/Private/StorageServerPlatformFile.cpp
--- a/Engine/Source/Runtime/StorageServerClient/Private/StorageServerPlatformFile.cpp
+++ b/Engine/Source/Runtime/StorageServerClient/Private/StorageServerPlatformFile.cpp
@@ -255,19 +255,10 @@ TUniquePtr<FArchive> FStorageServerPlatformFile::TryFindProjectStoreMarkerFile(I
 	{
 		return nullptr;
 	}
-	FString RootDir = FPaths::ConvertRelativePathToFull(FPaths::RootDir());
-	FString CookedDir = FPaths::ConvertRelativePathToFull(FPaths::Combine(*FPaths::ProjectDir(), TEXT("Saved"), TEXT("Cooked"), FPlatformProperties::PlatformName()));
-	TArray<FString> PotentialProjectStorePaths;
-	PotentialProjectStorePaths.Add(RootDir);
-	PotentialProjectStorePaths.Add(CookedDir);
-
-	for (const FString& ProjectStorePath : PotentialProjectStorePaths)
+	const TCHAR* ProjectMarkerPath = TEXT("../../../.projectstore");
+	if (IFileHandle* ProjectStoreMarkerHandle = Inner->OpenRead(ProjectMarkerPath); ProjectStoreMarkerHandle != nullptr)
 	{
-		FString ProjectMarkerPath = ProjectStorePath / TEXT(".projectstore");
-		if (IFileHandle* ProjectStoreMarkerHandle = Inner->OpenRead(*ProjectMarkerPath); ProjectStoreMarkerHandle != nullptr)
-		{
-			return TUniquePtr<FArchive>(new FArchiveFileReaderGeneric(ProjectStoreMarkerHandle, *ProjectMarkerPath, ProjectStoreMarkerHandle->Size()));
-		}
+		return TUniquePtr<FArchive>(new FArchiveFileReaderGeneric(ProjectStoreMarkerHandle, ProjectMarkerPath, ProjectStoreMarkerHandle->Size()));
 	}
 	return nullptr;
 }
