Rework zenserver host address detection when using COTF
Make sure .projectstore file is properly staged in all scenarios

diff --git a/Engine/Config/BasePakFileRules.ini b/Engine/Config/BasePakFileRules.ini
--- a/Engine/Config/BasePakFileRules.ini
+++ b/Engine/Config/BasePakFileRules.ini
@@ -22,7 +22,3 @@ bOverrideChunkManifest=true
 +Files=".../Engine/Content/EngineMaterials/DefaultCalibrationGrayscale.*"
 +Files=".../Engine/Content/EngineMaterials/PPM_DefaultCalibrationColor.*"
 +Files=".../Engine/Content/EngineMaterials/PPM_DefaultCalibrationGrayscale.*"
-
-[ExcludeProjectStoreMarkerFile]
-bExcludeFromPaks=true
-+Files=".../.projectstore"
diff --git a/Engine/Source/Developer/CookOnTheFlyNetServer/Private/CookOnTheFlyNetServerBase.cpp b/Engine/Source/Developer/CookOnTheFlyNetServer/Private/CookOnTheFlyNetServerBase.cpp
--- a/Engine/Source/Developer/CookOnTheFlyNetServer/Private/CookOnTheFlyNetServerBase.cpp
+++ b/Engine/Source/Developer/CookOnTheFlyNetServer/Private/CookOnTheFlyNetServerBase.cpp
@@ -5,6 +5,8 @@
 #include "Interfaces/ITargetPlatform.h"
 #include "Serialization/ArrayReader.h"
 #include "HAL/RunnableThread.h"
+#include "SocketSubsystem.h"
+#include "IPAddress.h"
 
 FCookOnTheFlyClientConnectionBase::FCookOnTheFlyClientConnectionBase(FCookOnTheFlyNetworkServerBase& InOwner)
 	: Owner(InOwner)
@@ -95,6 +97,29 @@ bool FCookOnTheFlyClientConnectionBase::Initialize()
 	HandshakeResponsePayload << ZenHostName;
 	HandshakeResponsePayload << ZenHostPort;
 
+	if (ZenHostName == "localhost")
+	{
+		TArray<FString> RemoteHostNames;
+		ISocketSubsystem* SocketSubsystem = ISocketSubsystem::Get();
+		if (SocketSubsystem != nullptr)
+		{
+			TArray<TSharedPtr<FInternetAddr>> Addresses;
+			if (SocketSubsystem->GetLocalAdapterAddresses(Addresses))
+			{
+				for (const TSharedPtr<FInternetAddr>& Address : Addresses)
+				{
+					RemoteHostNames.Add(Address->ToString(false));
+				}
+			}
+		}
+		int32 RemoteHostNameCount = RemoteHostNames.Num();
+		HandshakeResponsePayload << RemoteHostNameCount;
+		for (FString RemoteHostName : RemoteHostNames)
+		{
+			HandshakeResponsePayload << RemoteHostName;
+		}
+	}
+
 	if (!SendPayload(HandshakeResponsePayload))
 	{
 		return false;
diff --git a/Engine/Source/Programs/AutomationTool/Scripts/CopyBuildToStagingDirectory.Automation.cs b/Engine/Source/Programs/AutomationTool/Scripts/CopyBuildToStagingDirectory.Automation.cs
--- a/Engine/Source/Programs/AutomationTool/Scripts/CopyBuildToStagingDirectory.Automation.cs
+++ b/Engine/Source/Programs/AutomationTool/Scripts/CopyBuildToStagingDirectory.Automation.cs
@@ -1417,6 +1417,15 @@ namespace AutomationScripts
 								continue;
 							}
 
+							if (CookedFile.HasExtension(".projectstore"))
+							{
+								if (!ShouldCreatePak(Params, SC))
+								{
+									SC.StageFile(StagedFileType.NonUFS, CookedFile, new StagedFileReference(CookedFile.MakeRelativeTo(SC.PlatformCookDir)));
+								}
+								continue;
+							}
+
 							// json files have never been staged
 							// metallib files cannot *currently* be staged as UFS as the Metal API needs to mmap them from files on disk in order to function efficiently
 							if (!CookedFile.HasExtension(".json") && !CookedFile.HasExtension(".metallib") && !CookedFile.HasExtension(".utoc") && !CookedFile.HasExtension(".ucas"))
@@ -4624,7 +4633,7 @@ namespace AutomationScripts
 				}
 			}
 
-		string BaseManifestFileName = CombinePaths(CmdEnv.LogFolder, "FinalCopy" + (SC.DedicatedServer ? "_Server" : "") + SC.StageTargetPlatform.PlatformType.ToString());
+			string BaseManifestFileName = CombinePaths(CmdEnv.LogFolder, "FinalCopy" + (SC.DedicatedServer ? "_Server" : "") + SC.StageTargetPlatform.PlatformType.ToString());
 			DumpManifest(SC.FilesToStage.UFSFiles, BaseManifestFileName + "_UFSFiles.txt");
 
 			DumpManifest(SC.FilesToStage.NonUFSFiles, BaseManifestFileName + "_NonUFSFiles.txt");
diff --git a/Engine/Source/Runtime/CookOnTheFly/Internal/CookOnTheFly.h b/Engine/Source/Runtime/CookOnTheFly/Internal/CookOnTheFly.h
--- a/Engine/Source/Runtime/CookOnTheFly/Internal/CookOnTheFly.h
+++ b/Engine/Source/Runtime/CookOnTheFly/Internal/CookOnTheFly.h
@@ -307,7 +307,7 @@ public:
 
 	virtual const FString& GetPlatformName() const = 0;
 
-	virtual const FString& GetZenHostName() const = 0;
+	virtual const TArray<FString> GetZenHostNames() const = 0;
 
 	virtual const uint16 GetZenHostPort() const = 0;
 
diff --git a/Engine/Source/Runtime/CookOnTheFly/Private/CookOnTheFlyServerConnection.cpp b/Engine/Source/Runtime/CookOnTheFly/Private/CookOnTheFlyServerConnection.cpp
--- a/Engine/Source/Runtime/CookOnTheFly/Private/CookOnTheFlyServerConnection.cpp
+++ b/Engine/Source/Runtime/CookOnTheFly/Private/CookOnTheFlyServerConnection.cpp
@@ -59,15 +59,28 @@ public:
 			UE_LOG(LogCotfServerConnection, Fatal, TEXT("Couldn't handshake with server"));
 			return;
 		}
+		FString ZenHostName;
 		ResponsePayload << PlatformName;
 		ResponsePayload << ZenProjectName;
 		ResponsePayload << ZenHostName;
 		ResponsePayload << ZenHostPort;
 
-		// If our server is running its zenserver locally, we get "localhost" as a reply and we should use the same host for the zenserver
+		// If our server is running its zenserver locally, we get "localhost" as a reply and need to read the remote host name
 		if (ZenHostName == "localhost")
 		{
-			ZenHostName = Transport->GetHostName();
+			int32 NumRemoteHostNames = 0;
+			ResponsePayload << NumRemoteHostNames;
+			ZenHostNames.Reserve(NumRemoteHostNames);
+			for (int32 RemoteHostIndex = 0; RemoteHostIndex < NumRemoteHostNames; RemoteHostIndex++)
+			{
+				FString RemoteHostName;
+				ResponsePayload << RemoteHostName;
+				ZenHostNames.Add(RemoteHostName);
+			}
+		}
+		else
+		{
+			ZenHostNames.Add(ZenHostName);
 		}
 
 		FCoreDelegates::OnEnginePreExit.AddRaw(this, &FCookOnTheFlyServerConnection::OnEnginePreExit);
@@ -101,9 +114,9 @@ public:
 		return ZenProjectName;
 	}
 
-	virtual const FString& GetZenHostName() const override
+	virtual const TArray<FString> GetZenHostNames() const override
 	{
-		return ZenHostName;
+		return ZenHostNames;
 	}
 
 	virtual const uint16 GetZenHostPort() const override
@@ -325,7 +338,7 @@ private:
 	const bool bIsSingleThreaded;
 	FString PlatformName;
 	FString ZenProjectName;
-	FString ZenHostName;
+	TArray<FString> ZenHostNames;
 	uint16 ZenHostPort;
 
 	FCriticalSection RequestsCriticalSection;
diff --git a/Engine/Source/Runtime/CookOnTheFly/Private/CookOnTheFlyServerConnection.h b/Engine/Source/Runtime/CookOnTheFly/Private/CookOnTheFlyServerConnection.h
--- a/Engine/Source/Runtime/CookOnTheFly/Private/CookOnTheFlyServerConnection.h
+++ b/Engine/Source/Runtime/CookOnTheFly/Private/CookOnTheFlyServerConnection.h
@@ -17,8 +17,6 @@ public:
 	virtual bool SendPayload(const TArray<uint8>& Payload) = 0;
 	virtual bool HasPendingPayload() = 0;
 	virtual bool ReceivePayload(FArrayReader& Payload) = 0;
-	virtual FString GetHostName() const = 0;
-
 };
 
 UE::Cook::ICookOnTheFlyServerConnection* MakeCookOnTheFlyServerConnection(TUniquePtr<ICookOnTheFlyServerTransport> InTransport, const FString& InHost);
diff --git a/Engine/Source/Runtime/CookOnTheFly/Private/PlatformTransport.cpp b/Engine/Source/Runtime/CookOnTheFly/Private/PlatformTransport.cpp
--- a/Engine/Source/Runtime/CookOnTheFly/Private/PlatformTransport.cpp
+++ b/Engine/Source/Runtime/CookOnTheFly/Private/PlatformTransport.cpp
@@ -115,11 +115,6 @@ bool FPlatformTransport::HasPendingPayload()
 	return false;
 }
 
-FString FPlatformTransport::GetHostName() const
-{
-	return "platform://";
-}
-
 void FPlatformTransport::Disconnect()
 {
 	IPlatformHostCommunication& HostCommunication = FPlatformMisc::GetPlatformHostCommunication();
diff --git a/Engine/Source/Runtime/CookOnTheFly/Private/PlatformTransport.h b/Engine/Source/Runtime/CookOnTheFly/Private/PlatformTransport.h
--- a/Engine/Source/Runtime/CookOnTheFly/Private/PlatformTransport.h
+++ b/Engine/Source/Runtime/CookOnTheFly/Private/PlatformTransport.h
@@ -22,7 +22,6 @@ public:
 	virtual bool SendPayload(const TArray<uint8>& Payload) override;
 	virtual bool ReceivePayload(FArrayReader& Payload) override;
 	virtual bool HasPendingPayload() override;
-	virtual FString GetHostName() const override;
 	virtual void Disconnect() override;
 	//~ End ICookOnTheFlyServerTransport interface
 
diff --git a/Engine/Source/Runtime/CookOnTheFly/Private/TCPTransport.cpp b/Engine/Source/Runtime/CookOnTheFly/Private/TCPTransport.cpp
--- a/Engine/Source/Runtime/CookOnTheFly/Private/TCPTransport.cpp
+++ b/Engine/Source/Runtime/CookOnTheFly/Private/TCPTransport.cpp
@@ -96,11 +96,6 @@ bool FTCPTransport::HasPendingPayload()
 #endif
 }
 
-FString FTCPTransport::GetHostName() const
-{
-	return HostName;
-}
-
 void FTCPTransport::Disconnect()
 {
 #if USE_MCSOCKET_FOR_NFS
diff --git a/Engine/Source/Runtime/CookOnTheFly/Private/TCPTransport.h b/Engine/Source/Runtime/CookOnTheFly/Private/TCPTransport.h
--- a/Engine/Source/Runtime/CookOnTheFly/Private/TCPTransport.h
+++ b/Engine/Source/Runtime/CookOnTheFly/Private/TCPTransport.h
@@ -19,7 +19,6 @@ public:
 	virtual bool ReceivePayload(FArrayReader& Payload) override;
 	virtual bool HasPendingPayload() override;
 	virtual void Disconnect() override;
-	virtual FString GetHostName() const override;
 
 private: 
 
diff --git a/Engine/Source/Runtime/StorageServerClient/Private/StorageServerPlatformFile.cpp b/Engine/Source/Runtime/StorageServerClient/Private/StorageServerPlatformFile.cpp
--- a/Engine/Source/Runtime/StorageServerClient/Private/StorageServerPlatformFile.cpp
+++ b/Engine/Source/Runtime/StorageServerClient/Private/StorageServerPlatformFile.cpp
@@ -255,10 +255,25 @@ TUniquePtr<FArchive> FStorageServerPlatformFile::TryFindProjectStoreMarkerFile(I
 	{
 		return nullptr;
 	}
-	const TCHAR* ProjectMarkerPath = TEXT("../../../.projectstore");
-	if (IFileHandle* ProjectStoreMarkerHandle = Inner->OpenRead(ProjectMarkerPath); ProjectStoreMarkerHandle != nullptr)
+
+	FString RelativeStagedPath = TEXT("../../../");
+	FString RootPath = FPaths::RootDir();
+	FString PlatformName = FPlatformProperties::PlatformName();
+	FString CookedOutputPath = FPaths::Combine(FPaths::ProjectDir(), TEXT("Saved"), TEXT("Cooked"), PlatformName);
+
+	TArray<FString> PotentialProjectStorePaths;
+	PotentialProjectStorePaths.Add(RelativeStagedPath);
+	PotentialProjectStorePaths.Add(CookedOutputPath);
+	PotentialProjectStorePaths.Add(RootPath);
+
+	for (const FString& ProjectStorePath : PotentialProjectStorePaths)
 	{
-		return TUniquePtr<FArchive>(new FArchiveFileReaderGeneric(ProjectStoreMarkerHandle, ProjectMarkerPath, ProjectStoreMarkerHandle->Size()));
+		FString ProjectMarkerPath = ProjectStorePath / TEXT(".projectstore");
+		if (IFileHandle* ProjectStoreMarkerHandle = Inner->OpenRead(*ProjectMarkerPath); ProjectStoreMarkerHandle != nullptr)
+		{
+			UE_LOG(LogStorageServerPlatformFile, Display, TEXT("Found '%s'"), *ProjectMarkerPath);
+			return TUniquePtr<FArchive>(new FArchiveFileReaderGeneric(ProjectStoreMarkerHandle, *ProjectMarkerPath, ProjectStoreMarkerHandle->Size()));
+		}
 	}
 	return nullptr;
 }
@@ -270,7 +285,7 @@ bool FStorageServerPlatformFile::ShouldBeUsed(IPlatformFile* Inner, const TCHAR*
 	TSharedPtr<UE::Cook::ICookOnTheFlyServerConnection> DefaultConnection = CookOnTheFlyModule.GetDefaultServerConnection();
 	if (DefaultConnection.IsValid() && !DefaultConnection->GetZenProjectName().IsEmpty())
 	{
-		HostAddrs.Add(DefaultConnection->GetZenHostName());
+		HostAddrs.Append(DefaultConnection->GetZenHostNames());
 		HostPort = DefaultConnection->GetZenHostPort();
 		return true;
 	}
@@ -296,19 +311,23 @@ bool FStorageServerPlatformFile::ShouldBeUsed(IPlatformFile* Inner, const TCHAR*
 			}
 
 			HostPort = ZenServerField["hostport"].AsUInt16(HostPort);
-			UE_LOG(LogStorageServerPlatformFile, Display, TEXT("Found .projectstore"));
+			UE_LOG(LogStorageServerPlatformFile, Display, TEXT("Using connection settings from .projectstore: HostAddrs='%s' and HostPort='%d'"), *FString::Join(HostAddrs, TEXT("+")), HostPort);
 		}
 	}
 
 	FString Host;
 	if (FParse::Value(FCommandLine::Get(), TEXT("-ZenStoreHost="), Host))
 	{
+		UE_LOG(LogStorageServerPlatformFile, Display, TEXT("Adding connection settings from command line: -ZenStoreHost='%s'"), *Host);
 		if (!Host.ParseIntoArray(HostAddrs, TEXT("+"), true))
 		{
 			HostAddrs.Add(Host);
 		}
 	}
-	FParse::Value(CmdLine, TEXT("-ZenStorePort="), HostPort);
+	if (FParse::Value(CmdLine, TEXT("-ZenStorePort="), HostPort))
+	{
+		UE_LOG(LogStorageServerPlatformFile, Display, TEXT("Using connection settings from command line: -ZenStorePort='%d'"), HostPort);
+	}
 	return HostAddrs.Num() > 0;
 }
 
@@ -328,11 +347,18 @@ bool FStorageServerPlatformFile::Initialize(IPlatformFile* Inner, const TCHAR* C
 			{
 				ServerProject = FString(ZenServerField["projectid"].AsString());
 				ServerPlatform = FString(ZenServerField["oplogid"].AsString());
+				UE_LOG(LogStorageServerPlatformFile, Display, TEXT("Using settings from .projectstore: ServerProject='%s' and ServerPlatform='%s'"), *ServerProject, *ServerPlatform);
 			}
 		}
 	
-		FParse::Value(CmdLine, TEXT("-ZenStoreProject="), ServerProject);
-		FParse::Value(CmdLine, TEXT("-ZenStorePlatform="), ServerPlatform);
+		if (FParse::Value(CmdLine, TEXT("-ZenStoreProject="), ServerProject))
+		{
+			UE_LOG(LogStorageServerPlatformFile, Display, TEXT("Using settings from command line: -ZenStoreProject='%s'"), *ServerProject);
+		}
+		if (FParse::Value(CmdLine, TEXT("-ZenStorePlatform="), ServerPlatform))
+		{
+			UE_LOG(LogStorageServerPlatformFile, Display, TEXT("Using settings from command line: -ZenStorePlatform='%s'"), *ServerPlatform);
+		}
 		return true;
 	}
 	return false;
