[Backout] - CL26049351
Original CL Desc
-----------------------------------------------------------------
Rework zenserver host address detection when using COTF
Make sure .projectstore file is properly staged in all scenarios

diff --git a/Engine/Config/BasePakFileRules.ini b/Engine/Config/BasePakFileRules.ini
--- a/Engine/Config/BasePakFileRules.ini
+++ b/Engine/Config/BasePakFileRules.ini
@@ -22,3 +22,7 @@ bOverrideChunkManifest=true
 +Files=".../Engine/Content/EngineMaterials/DefaultCalibrationGrayscale.*"
 +Files=".../Engine/Content/EngineMaterials/PPM_DefaultCalibrationColor.*"
 +Files=".../Engine/Content/EngineMaterials/PPM_DefaultCalibrationGrayscale.*"
+
+[ExcludeProjectStoreMarkerFile]
+bExcludeFromPaks=true
++Files=".../.projectstore"
diff --git a/Engine/Source/Developer/CookOnTheFlyNetServer/Private/CookOnTheFlyNetServerBase.cpp b/Engine/Source/Developer/CookOnTheFlyNetServer/Private/CookOnTheFlyNetServerBase.cpp
--- a/Engine/Source/Developer/CookOnTheFlyNetServer/Private/CookOnTheFlyNetServerBase.cpp
+++ b/Engine/Source/Developer/CookOnTheFlyNetServer/Private/CookOnTheFlyNetServerBase.cpp
@@ -5,8 +5,6 @@
 #include "Interfaces/ITargetPlatform.h"
 #include "Serialization/ArrayReader.h"
 #include "HAL/RunnableThread.h"
-#include "SocketSubsystem.h"
-#include "IPAddress.h"
 
 FCookOnTheFlyClientConnectionBase::FCookOnTheFlyClientConnectionBase(FCookOnTheFlyNetworkServerBase& InOwner)
 	: Owner(InOwner)
@@ -97,29 +95,6 @@ bool FCookOnTheFlyClientConnectionBase::Initialize()
 	HandshakeResponsePayload << ZenHostName;
 	HandshakeResponsePayload << ZenHostPort;
 
-	if (ZenHostName == "localhost")
-	{
-		TArray<FString> RemoteHostNames;
-		ISocketSubsystem* SocketSubsystem = ISocketSubsystem::Get();
-		if (SocketSubsystem != nullptr)
-		{
-			TArray<TSharedPtr<FInternetAddr>> Addresses;
-			if (SocketSubsystem->GetLocalAdapterAddresses(Addresses))
-			{
-				for (const TSharedPtr<FInternetAddr>& Address : Addresses)
-				{
-					RemoteHostNames.Add(Address->ToString(false));
-				}
-			}
-		}
-		int32 RemoteHostNameCount = RemoteHostNames.Num();
-		HandshakeResponsePayload << RemoteHostNameCount;
-		for (FString RemoteHostName : RemoteHostNames)
-		{
-			HandshakeResponsePayload << RemoteHostName;
-		}
-	}
-
 	if (!SendPayload(HandshakeResponsePayload))
 	{
 		return false;
diff --git a/Engine/Source/Programs/AutomationTool/Scripts/CopyBuildToStagingDirectory.Automation.cs b/Engine/Source/Programs/AutomationTool/Scripts/CopyBuildToStagingDirectory.Automation.cs
--- a/Engine/Source/Programs/AutomationTool/Scripts/CopyBuildToStagingDirectory.Automation.cs
+++ b/Engine/Source/Programs/AutomationTool/Scripts/CopyBuildToStagingDirectory.Automation.cs
@@ -1417,15 +1417,6 @@ namespace AutomationScripts
 								continue;
 							}
 
-							if (CookedFile.HasExtension(".projectstore"))
-							{
-								if (!ShouldCreatePak(Params, SC))
-								{
-									SC.StageFile(StagedFileType.NonUFS, CookedFile, new StagedFileReference(CookedFile.MakeRelativeTo(SC.PlatformCookDir)));
-								}
-								continue;
-							}
-
 							// json files have never been staged
 							// metallib files cannot *currently* be staged as UFS as the Metal API needs to mmap them from files on disk in order to function efficiently
 							if (!CookedFile.HasExtension(".json") && !CookedFile.HasExtension(".metallib") && !CookedFile.HasExtension(".utoc") && !CookedFile.HasExtension(".ucas"))
@@ -4633,7 +4624,7 @@ namespace AutomationScripts
 				}
 			}
 
-			string BaseManifestFileName = CombinePaths(CmdEnv.LogFolder, "FinalCopy" + (SC.DedicatedServer ? "_Server" : "") + SC.StageTargetPlatform.PlatformType.ToString());
+		string BaseManifestFileName = CombinePaths(CmdEnv.LogFolder, "FinalCopy" + (SC.DedicatedServer ? "_Server" : "") + SC.StageTargetPlatform.PlatformType.ToString());
 			DumpManifest(SC.FilesToStage.UFSFiles, BaseManifestFileName + "_UFSFiles.txt");
 
 			DumpManifest(SC.FilesToStage.NonUFSFiles, BaseManifestFileName + "_NonUFSFiles.txt");
diff --git a/Engine/Source/Runtime/CookOnTheFly/Internal/CookOnTheFly.h b/Engine/Source/Runtime/CookOnTheFly/Internal/CookOnTheFly.h
--- a/Engine/Source/Runtime/CookOnTheFly/Internal/CookOnTheFly.h
+++ b/Engine/Source/Runtime/CookOnTheFly/Internal/CookOnTheFly.h
@@ -307,7 +307,7 @@ public:
 
 	virtual const FString& GetPlatformName() const = 0;
 
-	virtual const TArray<FString> GetZenHostNames() const = 0;
+	virtual const FString& GetZenHostName() const = 0;
 
 	virtual const uint16 GetZenHostPort() const = 0;
 
diff --git a/Engine/Source/Runtime/CookOnTheFly/Private/CookOnTheFlyServerConnection.cpp b/Engine/Source/Runtime/CookOnTheFly/Private/CookOnTheFlyServerConnection.cpp
--- a/Engine/Source/Runtime/CookOnTheFly/Private/CookOnTheFlyServerConnection.cpp
+++ b/Engine/Source/Runtime/CookOnTheFly/Private/CookOnTheFlyServerConnection.cpp
@@ -59,28 +59,15 @@ public:
 			UE_LOG(LogCotfServerConnection, Fatal, TEXT("Couldn't handshake with server"));
 			return;
 		}
-		FString ZenHostName;
 		ResponsePayload << PlatformName;
 		ResponsePayload << ZenProjectName;
 		ResponsePayload << ZenHostName;
 		ResponsePayload << ZenHostPort;
 
-		// If our server is running its zenserver locally, we get "localhost" as a reply and need to read the remote host name
+		// If our server is running its zenserver locally, we get "localhost" as a reply and we should use the same host for the zenserver
 		if (ZenHostName == "localhost")
 		{
-			int32 NumRemoteHostNames = 0;
-			ResponsePayload << NumRemoteHostNames;
-			ZenHostNames.Reserve(NumRemoteHostNames);
-			for (int32 RemoteHostIndex = 0; RemoteHostIndex < NumRemoteHostNames; RemoteHostIndex++)
-			{
-				FString RemoteHostName;
-				ResponsePayload << RemoteHostName;
-				ZenHostNames.Add(RemoteHostName);
-			}
-		}
-		else
-		{
-			ZenHostNames.Add(ZenHostName);
+			ZenHostName = Transport->GetHostName();
 		}
 
 		FCoreDelegates::OnEnginePreExit.AddRaw(this, &FCookOnTheFlyServerConnection::OnEnginePreExit);
@@ -114,9 +101,9 @@ public:
 		return ZenProjectName;
 	}
 
-	virtual const TArray<FString> GetZenHostNames() const override
+	virtual const FString& GetZenHostName() const override
 	{
-		return ZenHostNames;
+		return ZenHostName;
 	}
 
 	virtual const uint16 GetZenHostPort() const override
@@ -338,7 +325,7 @@ private:
 	const bool bIsSingleThreaded;
 	FString PlatformName;
 	FString ZenProjectName;
-	TArray<FString> ZenHostNames;
+	FString ZenHostName;
 	uint16 ZenHostPort;
 
 	FCriticalSection RequestsCriticalSection;
diff --git a/Engine/Source/Runtime/CookOnTheFly/Private/CookOnTheFlyServerConnection.h b/Engine/Source/Runtime/CookOnTheFly/Private/CookOnTheFlyServerConnection.h
--- a/Engine/Source/Runtime/CookOnTheFly/Private/CookOnTheFlyServerConnection.h
+++ b/Engine/Source/Runtime/CookOnTheFly/Private/CookOnTheFlyServerConnection.h
@@ -17,6 +17,8 @@ public:
 	virtual bool SendPayload(const TArray<uint8>& Payload) = 0;
 	virtual bool HasPendingPayload() = 0;
 	virtual bool ReceivePayload(FArrayReader& Payload) = 0;
+	virtual FString GetHostName() const = 0;
+
 };
 
 UE::Cook::ICookOnTheFlyServerConnection* MakeCookOnTheFlyServerConnection(TUniquePtr<ICookOnTheFlyServerTransport> InTransport, const FString& InHost);
diff --git a/Engine/Source/Runtime/CookOnTheFly/Private/PlatformTransport.cpp b/Engine/Source/Runtime/CookOnTheFly/Private/PlatformTransport.cpp
--- a/Engine/Source/Runtime/CookOnTheFly/Private/PlatformTransport.cpp
+++ b/Engine/Source/Runtime/CookOnTheFly/Private/PlatformTransport.cpp
@@ -115,6 +115,11 @@ bool FPlatformTransport::HasPendingPayload()
 	return false;
 }
 
+FString FPlatformTransport::GetHostName() const
+{
+	return "platform://";
+}
+
 void FPlatformTransport::Disconnect()
 {
 	IPlatformHostCommunication& HostCommunication = FPlatformMisc::GetPlatformHostCommunication();
diff --git a/Engine/Source/Runtime/CookOnTheFly/Private/PlatformTransport.h b/Engine/Source/Runtime/CookOnTheFly/Private/PlatformTransport.h
--- a/Engine/Source/Runtime/CookOnTheFly/Private/PlatformTransport.h
+++ b/Engine/Source/Runtime/CookOnTheFly/Private/PlatformTransport.h
@@ -22,6 +22,7 @@ public:
 	virtual bool SendPayload(const TArray<uint8>& Payload) override;
 	virtual bool ReceivePayload(FArrayReader& Payload) override;
 	virtual bool HasPendingPayload() override;
+	virtual FString GetHostName() const override;
 	virtual void Disconnect() override;
 	//~ End ICookOnTheFlyServerTransport interface
 
diff --git a/Engine/Source/Runtime/CookOnTheFly/Private/TCPTransport.cpp b/Engine/Source/Runtime/CookOnTheFly/Private/TCPTransport.cpp
--- a/Engine/Source/Runtime/CookOnTheFly/Private/TCPTransport.cpp
+++ b/Engine/Source/Runtime/CookOnTheFly/Private/TCPTransport.cpp
@@ -96,6 +96,11 @@ bool FTCPTransport::HasPendingPayload()
 #endif
 }
 
+FString FTCPTransport::GetHostName() const
+{
+	return HostName;
+}
+
 void FTCPTransport::Disconnect()
 {
 #if USE_MCSOCKET_FOR_NFS
diff --git a/Engine/Source/Runtime/CookOnTheFly/Private/TCPTransport.h b/Engine/Source/Runtime/CookOnTheFly/Private/TCPTransport.h
--- a/Engine/Source/Runtime/CookOnTheFly/Private/TCPTransport.h
+++ b/Engine/Source/Runtime/CookOnTheFly/Private/TCPTransport.h
@@ -19,6 +19,7 @@ public:
 	virtual bool ReceivePayload(FArrayReader& Payload) override;
 	virtual bool HasPendingPayload() override;
 	virtual void Disconnect() override;
+	virtual FString GetHostName() const override;
 
 private: 
 
diff --git a/Engine/Source/Runtime/StorageServerClient/Private/StorageServerPlatformFile.cpp b/Engine/Source/Runtime/StorageServerClient/Private/StorageServerPlatformFile.cpp
--- a/Engine/Source/Runtime/StorageServerClient/Private/StorageServerPlatformFile.cpp
+++ b/Engine/Source/Runtime/StorageServerClient/Private/StorageServerPlatformFile.cpp
@@ -255,25 +255,10 @@ TUniquePtr<FArchive> FStorageServerPlatformFile::TryFindProjectStoreMarkerFile(I
 	{
 		return nullptr;
 	}
-
-	FString RelativeStagedPath = TEXT("../../../");
-	FString RootPath = FPaths::RootDir();
-	FString PlatformName = FPlatformProperties::PlatformName();
-	FString CookedOutputPath = FPaths::Combine(FPaths::ProjectDir(), TEXT("Saved"), TEXT("Cooked"), PlatformName);
-
-	TArray<FString> PotentialProjectStorePaths;
-	PotentialProjectStorePaths.Add(RelativeStagedPath);
-	PotentialProjectStorePaths.Add(CookedOutputPath);
-	PotentialProjectStorePaths.Add(RootPath);
-
-	for (const FString& ProjectStorePath : PotentialProjectStorePaths)
+	const TCHAR* ProjectMarkerPath = TEXT("../../../.projectstore");
+	if (IFileHandle* ProjectStoreMarkerHandle = Inner->OpenRead(ProjectMarkerPath); ProjectStoreMarkerHandle != nullptr)
 	{
-		FString ProjectMarkerPath = ProjectStorePath / TEXT(".projectstore");
-		if (IFileHandle* ProjectStoreMarkerHandle = Inner->OpenRead(*ProjectMarkerPath); ProjectStoreMarkerHandle != nullptr)
-		{
-			UE_LOG(LogStorageServerPlatformFile, Display, TEXT("Found '%s'"), *ProjectMarkerPath);
-			return TUniquePtr<FArchive>(new FArchiveFileReaderGeneric(ProjectStoreMarkerHandle, *ProjectMarkerPath, ProjectStoreMarkerHandle->Size()));
-		}
+		return TUniquePtr<FArchive>(new FArchiveFileReaderGeneric(ProjectStoreMarkerHandle, ProjectMarkerPath, ProjectStoreMarkerHandle->Size()));
 	}
 	return nullptr;
 }
@@ -285,7 +270,7 @@ bool FStorageServerPlatformFile::ShouldBeUsed(IPlatformFile* Inner, const TCHAR*
 	TSharedPtr<UE::Cook::ICookOnTheFlyServerConnection> DefaultConnection = CookOnTheFlyModule.GetDefaultServerConnection();
 	if (DefaultConnection.IsValid() && !DefaultConnection->GetZenProjectName().IsEmpty())
 	{
-		HostAddrs.Append(DefaultConnection->GetZenHostNames());
+		HostAddrs.Add(DefaultConnection->GetZenHostName());
 		HostPort = DefaultConnection->GetZenHostPort();
 		return true;
 	}
@@ -311,23 +296,19 @@ bool FStorageServerPlatformFile::ShouldBeUsed(IPlatformFile* Inner, const TCHAR*
 			}
 
 			HostPort = ZenServerField["hostport"].AsUInt16(HostPort);
-			UE_LOG(LogStorageServerPlatformFile, Display, TEXT("Using connection settings from .projectstore: HostAddrs='%s' and HostPort='%d'"), *FString::Join(HostAddrs, TEXT("+")), HostPort);
+			UE_LOG(LogStorageServerPlatformFile, Display, TEXT("Found .projectstore"));
 		}
 	}
 
 	FString Host;
 	if (FParse::Value(FCommandLine::Get(), TEXT("-ZenStoreHost="), Host))
 	{
-		UE_LOG(LogStorageServerPlatformFile, Display, TEXT("Adding connection settings from command line: -ZenStoreHost='%s'"), *Host);
 		if (!Host.ParseIntoArray(HostAddrs, TEXT("+"), true))
 		{
 			HostAddrs.Add(Host);
 		}
 	}
-	if (FParse::Value(CmdLine, TEXT("-ZenStorePort="), HostPort))
-	{
-		UE_LOG(LogStorageServerPlatformFile, Display, TEXT("Using connection settings from command line: -ZenStorePort='%d'"), HostPort);
-	}
+	FParse::Value(CmdLine, TEXT("-ZenStorePort="), HostPort);
 	return HostAddrs.Num() > 0;
 }
 
@@ -347,18 +328,11 @@ bool FStorageServerPlatformFile::Initialize(IPlatformFile* Inner, const TCHAR* C
 			{
 				ServerProject = FString(ZenServerField["projectid"].AsString());
 				ServerPlatform = FString(ZenServerField["oplogid"].AsString());
-				UE_LOG(LogStorageServerPlatformFile, Display, TEXT("Using settings from .projectstore: ServerProject='%s' and ServerPlatform='%s'"), *ServerProject, *ServerPlatform);
 			}
 		}
 	
-		if (FParse::Value(CmdLine, TEXT("-ZenStoreProject="), ServerProject))
-		{
-			UE_LOG(LogStorageServerPlatformFile, Display, TEXT("Using settings from command line: -ZenStoreProject='%s'"), *ServerProject);
-		}
-		if (FParse::Value(CmdLine, TEXT("-ZenStorePlatform="), ServerPlatform))
-		{
-			UE_LOG(LogStorageServerPlatformFile, Display, TEXT("Using settings from command line: -ZenStorePlatform='%s'"), *ServerPlatform);
-		}
+		FParse::Value(CmdLine, TEXT("-ZenStoreProject="), ServerProject);
+		FParse::Value(CmdLine, TEXT("-ZenStorePlatform="), ServerPlatform);
 		return true;
 	}
 	return false;
