Use the already loaded bytes available even if read overlaps boundary of buffer in FStorageServerFileHandle before

diff --git a/Engine/Source/Runtime/StorageServerClient/Private/StorageServerPlatformFile.cpp b/Engine/Source/Runtime/StorageServerClient/Private/StorageServerPlatformFile.cpp
--- a/Engine/Source/Runtime/StorageServerClient/Private/StorageServerPlatformFile.cpp
+++ b/Engine/Source/Runtime/StorageServerClient/Private/StorageServerPlatformFile.cpp
@@ -202,22 +202,33 @@ public:
 				FilePos += BytesRead;
 				return true;
 			}
+			return false;
 		}
 
-		if (FilePos < BufferStart || BufferEnd < FilePos + BytesToRead)
+		int64 BytesReadFromBuffer = 0;
+		if (FilePos >= BufferStart && FilePos < BufferEnd)
 		{
-			const int64 BytesRead = Owner.SendReadMessage(Buffer, FileChunkId, FilePos, BufferSize);
-			BufferStart = FilePos;
-			BufferEnd = BufferStart + BytesRead;
+			const int64 BufferOffset = FilePos - BufferStart;
+			check(BufferOffset < BufferSize);
+			BytesReadFromBuffer = FMath::Min(BufferSize - BufferOffset, BytesToRead);
+			FMemory::Memcpy(Destination, Buffer + BufferOffset, BytesReadFromBuffer);
+			if (BytesReadFromBuffer == BytesToRead)
+			{
+				FilePos += BytesReadFromBuffer;
+				return true;
+			}
 		}
 
-		int64 BufferOffset = FilePos - BufferStart;
-		check(BufferEnd > BufferOffset);
-		int64 BytesToReadFromBuffer = FMath::Min(BufferEnd - BufferOffset, BytesToRead);
-		FMemory::Memcpy(Destination, Buffer + BufferOffset, BytesToReadFromBuffer);
-		if (BytesToReadFromBuffer == BytesToRead)
+		const int64 BytesRead = Owner.SendReadMessage(Buffer, FileChunkId, FilePos + BytesReadFromBuffer, BufferSize);
+		BufferStart = FilePos + BytesReadFromBuffer;
+		BufferEnd = BufferStart + BytesRead;
+
+		const int64 BytesToReadFromBuffer = FMath::Min(BytesRead, BytesToRead - BytesReadFromBuffer);
+		FMemory::Memcpy(Destination + BytesReadFromBuffer, Buffer, BytesToReadFromBuffer);
+		BytesReadFromBuffer += BytesToReadFromBuffer;
+		if (BytesReadFromBuffer == BytesToRead)
 		{
-			FilePos += BytesToReadFromBuffer;
+			FilePos += BytesReadFromBuffer;
 			return true;
 		}
 		
