Target Iteration: Usability improvements to StorageServerConnection/File:
- Set a lower timeout on the duration of initial handshake (5.0 seconds) so that we don't stall for a long time trying to connect to an invalid host
- Provide a descriptive error messagebox when not in unattended mode and we fail to connect to any host
- Avoid attempting to connect to localhost on non-desktop platforms (they can never be the host, and if we're using IPv6 localhost, it can cause them to try an unnecessary name lookup).

diff --git a/Engine/Source/Runtime/StorageServerClient/Private/StorageServerConnection.cpp b/Engine/Source/Runtime/StorageServerClient/Private/StorageServerConnection.cpp
--- a/Engine/Source/Runtime/StorageServerClient/Private/StorageServerConnection.cpp
+++ b/Engine/Source/Runtime/StorageServerClient/Private/StorageServerConnection.cpp
@@ -506,7 +506,10 @@ int32 FStorageServerConnection::HandshakeRequest(TArrayView<const TSharedPtr<FIn
 		
 		UE_LOG(LogStorageServerConnection, Display, TEXT("Trying to handshake with Zen at '%s'"), *Addr->ToString(true));
 
-		FSocket* ConnectSocket = AcquireNewSocket();
+		// Handshakes are done with a limited connection timeout so that we can find out if the destination is unreachable
+		// in a timely manner.
+		const float ConnectionTimeoutSeconds = 5.0f;
+		FSocket* ConnectSocket = AcquireNewSocket(ConnectionTimeoutSeconds);
 		if (!ConnectSocket)
 		{
 			continue;
@@ -715,14 +718,30 @@ FSocket* FStorageServerConnection::AcquireSocketFromPool()
 	return nullptr;
 }
 
-FSocket* FStorageServerConnection::AcquireNewSocket()
+FSocket* FStorageServerConnection::AcquireNewSocket(float TimeoutSeconds)
 {
 	FSocket* Socket = SocketSubsystem.CreateSocket(NAME_Stream, TEXT("StorageServer"), ServerAddr->GetProtocolType());
 	check(Socket);
 
-	if (Socket->Connect(*ServerAddr))
+	if (TimeoutSeconds > 0.0f)
 	{
-		return Socket;
+		Socket->SetNonBlocking(true);
+		ON_SCOPE_EXIT
+		{
+			Socket->SetNonBlocking(false);
+		};
+
+		if (Socket->Connect(*ServerAddr) && Socket->Wait(ESocketWaitConditions::WaitForWrite, FTimespan::FromSeconds(TimeoutSeconds)))
+		{
+			return Socket;
+		}
+	}
+	else
+	{
+		if (Socket->Connect(*ServerAddr))
+		{
+			return Socket;
+		}
 	}
 
 	delete Socket;
diff --git a/Engine/Source/Runtime/StorageServerClient/Private/StorageServerPlatformFile.cpp b/Engine/Source/Runtime/StorageServerClient/Private/StorageServerPlatformFile.cpp
--- a/Engine/Source/Runtime/StorageServerClient/Private/StorageServerPlatformFile.cpp
+++ b/Engine/Source/Runtime/StorageServerClient/Private/StorageServerPlatformFile.cpp
@@ -1,22 +1,23 @@
 // Copyright Epic Games, Inc. All Rights Reserved.
 
 #include "StorageServerPlatformFile.h"
-#include "StorageServerIoDispatcherBackend.h"
-#include "HAL/IPlatformFileModule.h"
+#include "Algo/Replace.h"
+#include "CookOnTheFly.h"
+#include "CookOnTheFlyPackageStore.h"
 #include "HAL/FileManagerGeneric.h"
+#include "HAL/IPlatformFileModule.h"
+#include "Misc/App.h"
 #include "Misc/CommandLine.h"
+#include "Misc/CoreDelegates.h"
 #include "Misc/Paths.h"
 #include "Misc/ScopeRWLock.h"
-#include "StorageServerConnection.h"
-#include "Modules/ModuleManager.h"
 #include "Misc/StringBuilder.h"
-#include "Algo/Replace.h"
-#include "StorageServerPackageStore.h"
-#include "CookOnTheFlyPackageStore.h"
-#include "Misc/CoreDelegates.h"
 #include "Modules/ModuleManager.h"
-#include "CookOnTheFly.h"
+#include "Modules/ModuleManager.h"
 #include "Serialization/CompactBinarySerialization.h"
+#include "StorageServerConnection.h"
+#include "StorageServerIoDispatcherBackend.h"
+#include "StorageServerPackageStore.h"
 
 DEFINE_LOG_CATEGORY_STATIC(LogStorageServerPlatformFile, Log, All);
 
@@ -308,10 +309,12 @@ bool FStorageServerPlatformFile::ShouldBeUsed(IPlatformFile* Inner, const TCHAR*
 
 		if (FCbFieldView ZenServerField = ProjectStoreObject["zenserver"])
 		{
+#if PLATFORM_DESKTOP
 			if (FUtf8StringView HostName = ZenServerField["hostname"].AsString(); !HostName.IsEmpty())
 			{
 				HostAddrs.Add(FString(HostName));
 			}
+#endif
 			FCbArrayView RemoteHostNames = ZenServerField["remotehostnames"].AsArrayView();
 			for (FCbFieldView RemoteHostName : RemoteHostNames)
 			{
@@ -415,7 +418,18 @@ void FStorageServerPlatformFile::InitializeAfterProjectFilePath()
 	}
 	else
 	{
-		UE_LOG(LogStorageServerPlatformFile, Fatal, TEXT("Failed to initialize connection"));
+		if (!FApp::IsUnattended())
+		{
+			FText FailedConnectionTitle = NSLOCTEXT("StorageServer", "StorageServer_ConnectFailedTitle", "Failed to connect");
+			FText FailedConnectionText = FText::Format(NSLOCTEXT("StorageServer", "StorageServer_ConnectFailedText",
+				"Network data streaming failed to connect to any of the following data sources:\n\n{0}\n\n"
+				"This can be due to the sources being offline, invalid addresses, firewall blocking, or the sources being on a different network from this device. "
+				"If these issues can't be addressed, you can use an installed build without network data streaming. This process will now exit."),
+				FText::FromString(FString::Join(HostAddrs, TEXT("\n"))));
+			FPlatformMisc::MessageBoxExt(EAppMsgType::Ok, *FailedConnectionText.ToString(), *FailedConnectionTitle.ToString());
+		}
+
+		UE_LOG(LogStorageServerPlatformFile, Fatal, TEXT("Failed to initialize connection to %s"), *FString::Join(HostAddrs, TEXT("\n")));
 	}
 }
 
diff --git a/Engine/Source/Runtime/StorageServerClient/Public/StorageServerConnection.h b/Engine/Source/Runtime/StorageServerClient/Public/StorageServerConnection.h
--- a/Engine/Source/Runtime/StorageServerClient/Public/StorageServerConnection.h
+++ b/Engine/Source/Runtime/StorageServerClient/Public/StorageServerConnection.h
@@ -197,7 +197,7 @@ private:
 
 	int32 HandshakeRequest(TArrayView<const TSharedPtr<FInternetAddr>> HostAddresses);
 	FSocket* AcquireSocketFromPool();
-	FSocket* AcquireNewSocket();
+	FSocket* AcquireNewSocket(float TimeoutSeconds = -1.0f);
 	void ReleaseSocket(FSocket* Socket, bool bKeepAlive);
 
 	ISocketSubsystem& SocketSubsystem;
