CookOnTheFly2: Remove duplicate warnings for missing packages requested multiple times, and remove warnings for editor-only packages that are found due to unsolicited loads on the cooker but were not requested by the client. Add packagename to warnings about missing packages.
#rnx

diff --git a/Engine/Plugins/MemoryUsageQueries/Source/MemoryUsageQueries/Private/MemoryUsageQueries.cpp b/Engine/Plugins/MemoryUsageQueries/Source/MemoryUsageQueries/Private/MemoryUsageQueries.cpp
--- a/Engine/Plugins/MemoryUsageQueries/Source/MemoryUsageQueries/Private/MemoryUsageQueries.cpp
+++ b/Engine/Plugins/MemoryUsageQueries/Source/MemoryUsageQueries/Private/MemoryUsageQueries.cpp
@@ -1449,7 +1449,8 @@ class FPackageDependenciesLazyDatabase final
 			}
 
 			FPackageStoreEntry PackageEntry;
-			const EPackageStoreEntryStatus Status = FPackageStore::Get().GetPackageStoreEntry(PackageId, PackageEntry);
+			const EPackageStoreEntryStatus Status = FPackageStore::Get().GetPackageStoreEntry(PackageId, NAME_None,
+				PackageEntry);
 			if (Status == EPackageStoreEntryStatus::Ok)
 			{
 				// add package dependencies
diff --git a/Engine/Source/Editor/UnrealEd/Private/Cooker/IoStoreCookOnTheFlyRequestManager.cpp b/Engine/Source/Editor/UnrealEd/Private/Cooker/IoStoreCookOnTheFlyRequestManager.cpp
--- a/Engine/Source/Editor/UnrealEd/Private/Cooker/IoStoreCookOnTheFlyRequestManager.cpp
+++ b/Engine/Source/Editor/UnrealEd/Private/Cooker/IoStoreCookOnTheFlyRequestManager.cpp
@@ -138,6 +138,7 @@ private:
 		EPackageStoreEntryStatus RequestCook(UE::Cook::ICookOnTheFlyServer& InCookOnTheFlyServer, const FPackageId& PackageId, FPackageStoreEntryResource& OutEntry)
 		{
 			FPackage& Package = GetPackage(PackageId);
+			FName PackageName = PackageRegistry.Get(PackageId);
 			if (Package.Status == EPackageStatus::Cooked)
 			{
 				UE_LOG(LogCookOnTheFly, Verbose, TEXT("0x%llX was already cooked"), PackageId.ValueForDebugging());
@@ -146,7 +147,7 @@ private:
 			}
 			else if (Package.Status == EPackageStatus::Failed)
 			{
-				UE_LOG(LogCookOnTheFly, Verbose, TEXT("0x%llX was already failed"), PackageId.ValueForDebugging());
+				UE_LOG(LogCookOnTheFly, Warning, TEXT("Failed to cook package 0x%llX '%s'"), PackageId.ValueForDebugging(), *PackageName.ToString());
 				return EPackageStoreEntryStatus::Missing;
 			}
 			else if (Package.Status == EPackageStatus::Cooking)
@@ -154,7 +155,6 @@ private:
 				UE_LOG(LogCookOnTheFly, Verbose, TEXT("0x%llX was already cooking"), PackageId.ValueForDebugging());
 				return EPackageStoreEntryStatus::Pending;
 			}
-			FName PackageName = PackageRegistry.Get(PackageId);
 			if (PackageName.IsNone())
 			{
 				UE_LOG(LogCookOnTheFly, Warning, TEXT("Received cook request for unknown package 0x%llX"), PackageId.ValueForDebugging());
@@ -200,12 +200,12 @@ private:
 			}
 		}
 
-		void MarkAsFailed(FPackageId PackageId, UE::ZenCookOnTheFly::Messaging::FCompletedPackages& OutCompletedPackages)
+		void MarkAsFailed(FPackageId PackageId, FName PackageName, UE::ZenCookOnTheFly::Messaging::FCompletedPackages& OutCompletedPackages)
 		{
-			UE_LOG(LogCookOnTheFly, Warning, TEXT("0x%llX failed"), PackageId.ValueForDebugging());
+			// Unsolicited packages that fail to save because they contain editoronly data will be passed into this function
+			// Mark them as failed, but do not report them to the client and do not give a warning unless the client requests them.
 			FPackage& Package = GetPackage(PackageId);
 			Package.Status = EPackageStatus::Failed;
-			OutCompletedPackages.FailedPackages.Add(PackageId);
 			if (PackageCookedEvent)
 			{
 				PackageCookedEvent->Trigger();
@@ -692,7 +692,7 @@ private:
 			}
 			else
 			{
-				Context.MarkAsFailed(FPackageId::FromName(EventArgs.PackageName), NewCompletedPackages);
+				Context.MarkAsFailed(FPackageId::FromName(EventArgs.PackageName), EventArgs.PackageName, NewCompletedPackages);
 			}
 		}
 		BroadcastCompletedPackages(EventArgs.PlatformName, MoveTemp(NewCompletedPackages));
diff --git a/Engine/Source/Runtime/Core/Private/IO/PackageStore.cpp b/Engine/Source/Runtime/Core/Private/IO/PackageStore.cpp
--- a/Engine/Source/Runtime/Core/Private/IO/PackageStore.cpp
+++ b/Engine/Source/Runtime/Core/Private/IO/PackageStore.cpp
@@ -158,12 +158,13 @@ void FPackageStore::Mount(TSharedRef<IPackageStoreBackend> Backend, int32 Priori
 	Backend->OnMounted(BackendContext);
 }
 
-EPackageStoreEntryStatus FPackageStore::GetPackageStoreEntry(FPackageId PackageId, FPackageStoreEntry& OutPackageStoreEntry)
+EPackageStoreEntryStatus FPackageStore::GetPackageStoreEntry(FPackageId PackageId, FName PackageName,
+	FPackageStoreEntry& OutPackageStoreEntry)
 {
 	check(ThreadReadCount);
 	for (const FBackendAndPriority& Backend : Backends)
 	{
-		EPackageStoreEntryStatus Status = Backend.Value->GetPackageStoreEntry(PackageId, OutPackageStoreEntry);
+		EPackageStoreEntryStatus Status = Backend.Value->GetPackageStoreEntry(PackageId, PackageName, OutPackageStoreEntry);
 		if (Status >= EPackageStoreEntryStatus::Pending)
 		{
 			return Status;
diff --git a/Engine/Source/Runtime/Core/Public/IO/PackageStore.h b/Engine/Source/Runtime/Core/Public/IO/PackageStore.h
--- a/Engine/Source/Runtime/Core/Public/IO/PackageStore.h
+++ b/Engine/Source/Runtime/Core/Public/IO/PackageStore.h
@@ -120,7 +120,8 @@ public:
 	virtual bool DoesPackageExist(FPackageId PackageId) = 0;
 
 	/* Returns the package store entry data with export info and imported packages for the specified package ID. */
-	virtual EPackageStoreEntryStatus GetPackageStoreEntry(FPackageId PackageId, FPackageStoreEntry& OutPackageStoreEntry) = 0;
+	virtual EPackageStoreEntryStatus GetPackageStoreEntry(FPackageId PackageId, FName PackageName,
+		FPackageStoreEntry& OutPackageStoreEntry) = 0;
 
 	/* Returns the redirected package ID and source package name for the specified package ID if it's being redirected. */
 	virtual bool GetPackageRedirectInfo(FPackageId PackageId, FName& OutSourcePackageName, FPackageId& OutRedirectedToPackageId) = 0;
@@ -172,7 +173,8 @@ public:
 	virtual void EndRead() = 0;
 
 	/* Returns the package store entry data with export info and imported packages for the specified package ID. */
-	virtual EPackageStoreEntryStatus GetPackageStoreEntry(FPackageId PackageId, FPackageStoreEntry& OutPackageStoreEntry) = 0;
+	virtual EPackageStoreEntryStatus GetPackageStoreEntry(FPackageId PackageId, FName PackageName,
+		FPackageStoreEntry& OutPackageStoreEntry) = 0;
 
 	/* Returns the redirected package ID and source package name for the specified package ID if it's being redirected. */
 	virtual bool GetPackageRedirectInfo(FPackageId PackageId, FName& OutSourcePackageName, FPackageId& OutRedirectedToPackageId) = 0;
@@ -190,7 +192,8 @@ public:
 	CORE_API void Mount(TSharedRef<IPackageStoreBackend> Backend, int32 Priority = 0);
 
 	/* Returns the package store entry data with export info and imported packages for the specified package ID. */
-	CORE_API EPackageStoreEntryStatus GetPackageStoreEntry(FPackageId PackageId, FPackageStoreEntry& OutPackageStoreEntry);
+	CORE_API EPackageStoreEntryStatus GetPackageStoreEntry(FPackageId PackageId, FName PackageName, 
+		FPackageStoreEntry& OutPackageStoreEntry);
 
 	/* Returns the redirected package ID and source package name for the specified package ID if it's being redirected. */
 	CORE_API bool GetPackageRedirectInfo(FPackageId PackageId, FName& OutSourcePackageName, FPackageId& OutRedirectedToPackageId);
diff --git a/Engine/Source/Runtime/CoreUObject/Private/Serialization/AsyncLoading2.cpp b/Engine/Source/Runtime/CoreUObject/Private/Serialization/AsyncLoading2.cpp
--- a/Engine/Source/Runtime/CoreUObject/Private/Serialization/AsyncLoading2.cpp
+++ b/Engine/Source/Runtime/CoreUObject/Private/Serialization/AsyncLoading2.cpp
@@ -4074,7 +4074,8 @@ bool FAsyncLoadingThread2::CreateAsyncPackagesFromQueue(FAsyncLoadingThreadState
 	{
 		FAsyncPackage2* PendingPackage = *It;
 		FPackageStoreEntry PackageEntry;
-		EPackageStoreEntryStatus PendingPackageStatus = PackageStore.GetPackageStoreEntry(PendingPackage->Desc.PackageIdToLoad, PackageEntry);
+		EPackageStoreEntryStatus PendingPackageStatus = PackageStore.GetPackageStoreEntry(PendingPackage->Desc.PackageIdToLoad,
+			PendingPackage->Desc.UPackageName, PackageEntry);
 		if (PendingPackageStatus == EPackageStoreEntryStatus::Ok)
 		{
 			InitializeAsyncPackageFromPackageStore(ThreadState, &IoBatch, PendingPackage, PackageEntry);
@@ -4147,7 +4148,7 @@ bool FAsyncLoadingThread2::CreateAsyncPackagesFromQueue(FAsyncLoadingThreadState
 				}
 			}
 
-			PackageStatus = PackageStore.GetPackageStoreEntry(PackageIdToLoad, PackageEntry);
+			PackageStatus = PackageStore.GetPackageStoreEntry(PackageIdToLoad, UPackageName, PackageEntry);
 			if (PackageStatus == EPackageStoreEntryStatus::Missing)
 			{
 				// While there is an active load request for (InName=/Temp/PackageABC_abc, InPackageToLoadFrom=/Game/PackageABC), then allow these requests too:
@@ -4158,7 +4159,7 @@ bool FAsyncLoadingThread2::CreateAsyncPackagesFromQueue(FAsyncLoadingThreadState
 					PackageIdToLoad = Package->Desc.PackageIdToLoad;
 					Request.PackagePath = Package->Desc.PackagePathToLoad;
 					PackageNameToLoad = Request.PackagePath.GetPackageFName();
-					PackageStatus = PackageStore.GetPackageStoreEntry(PackageIdToLoad, PackageEntry);
+					PackageStatus = PackageStore.GetPackageStoreEntry(PackageIdToLoad, UPackageName, PackageEntry);
 				}
 			}
 
@@ -4758,7 +4759,8 @@ void FAsyncPackage2::ImportPackagesRecursiveInner(FAsyncLoadingThreadState2& Thr
 				ImportedPackageNameToLoad = NAME_None;
 			}
 		}
-		ImportedPackageStatus = PackageStore.GetPackageStoreEntry(ImportedPackageIdToLoad, ImportedPackageEntry);
+		ImportedPackageStatus = PackageStore.GetPackageStoreEntry(ImportedPackageIdToLoad, ImportedPackageUPackageName,
+			ImportedPackageEntry);
 
 		FPackagePath ImportedPackagePath;
 #if ALT2_ENABLE_LINKERLOAD_SUPPORT || WITH_EDITOR
diff --git a/Engine/Source/Runtime/PakFile/Private/FilePackageStore.cpp b/Engine/Source/Runtime/PakFile/Private/FilePackageStore.cpp
--- a/Engine/Source/Runtime/PakFile/Private/FilePackageStore.cpp
+++ b/Engine/Source/Runtime/PakFile/Private/FilePackageStore.cpp
@@ -49,7 +49,8 @@ void FFilePackageStoreBackend::EndRead()
 	EntriesLock.ReadUnlock();
 }
 
-EPackageStoreEntryStatus FFilePackageStoreBackend::GetPackageStoreEntry(FPackageId PackageId, FPackageStoreEntry& OutPackageStoreEntry)
+EPackageStoreEntryStatus FFilePackageStoreBackend::GetPackageStoreEntry(FPackageId PackageId, FName PackageName,
+	FPackageStoreEntry& OutPackageStoreEntry)
 {
 	const FFilePackageStoreEntry* FindEntry = StoreEntriesMap.FindRef(PackageId);
 	if (FindEntry)
diff --git a/Engine/Source/Runtime/PakFile/Private/FilePackageStore.h b/Engine/Source/Runtime/PakFile/Private/FilePackageStore.h
--- a/Engine/Source/Runtime/PakFile/Private/FilePackageStore.h
+++ b/Engine/Source/Runtime/PakFile/Private/FilePackageStore.h
@@ -26,7 +26,8 @@ public:
 
 	virtual void BeginRead() override;
 	virtual void EndRead() override;
-	virtual EPackageStoreEntryStatus GetPackageStoreEntry(FPackageId PackageId, FPackageStoreEntry& OutPackageStoreEntry) override;
+	virtual EPackageStoreEntryStatus GetPackageStoreEntry(FPackageId PackageId, FName PackageName,
+		FPackageStoreEntry& OutPackageStoreEntry) override;
 	virtual bool GetPackageRedirectInfo(FPackageId PackageId, FName& OutSourcePackageName, FPackageId& OutRedirectedToPackageId) override;
 
 	void Mount(const FIoContainerHeader* ContainerHeader, uint32 Order);
diff --git a/Engine/Source/Runtime/StorageServerClient/Private/CookOnTheFlyPackageStore.cpp b/Engine/Source/Runtime/StorageServerClient/Private/CookOnTheFlyPackageStore.cpp
--- a/Engine/Source/Runtime/StorageServerClient/Private/CookOnTheFlyPackageStore.cpp
+++ b/Engine/Source/Runtime/StorageServerClient/Private/CookOnTheFlyPackageStore.cpp
@@ -32,7 +32,8 @@ FCookOnTheFlyPackageStoreBackend::FCookOnTheFlyPackageStoreBackend(UE::Cook::ICo
 		UE_LOG(LogCookOnTheFly, Log, TEXT("Got '%d' cooked and '%d' failed packages from server"),
 			GetCookedPackagesResponse.CookedPackages.Num(), GetCookedPackagesResponse.FailedPackages.Num());
 
-		AddPackages(MoveTemp(GetCookedPackagesResponse.CookedPackages), MoveTemp(GetCookedPackagesResponse.FailedPackages));
+		AddPackages(MoveTemp(GetCookedPackagesResponse.CookedPackages), MoveTemp(GetCookedPackagesResponse.FailedPackages),
+			TArray<TPair<FPackageId, FName>>());
 
 		LastServerActivtyTime = FPlatformTime::Seconds();
 	}
@@ -54,7 +55,7 @@ void FCookOnTheFlyPackageStoreBackend::BeginRead()
 
 void FCookOnTheFlyPackageStoreBackend::EndRead()
 {
-	TArray<FPackageId> PackageIds = MoveTemp(RequestedPackageIds);
+	TArray<TPair<FPackageId, FName>> PackageIds = MoveTemp(RequestedPackageIds);
 	EntriesLock.ReadUnlock();
 
 	if (PackageIds.Num() > 0)
@@ -70,7 +71,8 @@ bool FCookOnTheFlyPackageStoreBackend::DoesPackageExist(FPackageId PackageId)
 	return EntryInfo.Status != EPackageStoreEntryStatus::Missing;
 }
 
-EPackageStoreEntryStatus FCookOnTheFlyPackageStoreBackend::GetPackageStoreEntry(FPackageId PackageId, FPackageStoreEntry& OutPackageStoreEntry)
+EPackageStoreEntryStatus FCookOnTheFlyPackageStoreBackend::GetPackageStoreEntry(FPackageId PackageId,
+	FName PackageName, FPackageStoreEntry& OutPackageStoreEntry)
 {
 	FEntryInfo& EntryInfo = PackageIdToEntryInfo.FindOrAdd(PackageId, FEntryInfo());
 
@@ -82,19 +84,25 @@ EPackageStoreEntryStatus FCookOnTheFlyPackageStoreBackend::GetPackageStoreEntry(
 	else if (EntryInfo.Status == EPackageStoreEntryStatus::None)
 	{
 		EntryInfo.Status = EPackageStoreEntryStatus::Pending;
-		RequestedPackageIds.Add(PackageId);
+		RequestedPackageIds.Add(TPair<FPackageId, FName>(PackageId, PackageName));
 	}
 
 	return EntryInfo.Status; 
 }
 
-void FCookOnTheFlyPackageStoreBackend::SendCookRequest(TArray<FPackageId> PackageIds)
+void FCookOnTheFlyPackageStoreBackend::SendCookRequest(TArray<TPair<FPackageId, FName>> PackageIdsAndNames)
 {
 	using namespace UE::Cook;
 	using namespace UE::ZenCookOnTheFly::Messaging;
 
 	LastClientActivtyTime = FPlatformTime::Seconds();
 	
+	TArray<FPackageId> PackageIds;
+	PackageIds.Reserve(PackageIdsAndNames.Num());
+	for (const TPair<FPackageId, FName>& Pair : PackageIdsAndNames)
+	{
+		PackageIds.Add(Pair.Key);
+	}
 	FCookOnTheFlyRequest Request(ECookOnTheFlyMessage::CookPackage);
 	Request.SetBodyTo(FCookPackageRequest { PackageIds });
 
@@ -103,7 +111,8 @@ void FCookOnTheFlyPackageStoreBackend::SendCookRequest(TArray<FPackageId> Packag
 
 	if (Response.IsOk())
 	{
-		AddPackages(MoveTemp(CookResponse.CookedPackages), MoveTemp(CookResponse.FailedPackages));
+		AddPackages(MoveTemp(CookResponse.CookedPackages), MoveTemp(CookResponse.FailedPackages),
+			MoveTemp(PackageIdsAndNames));
 	}
 	else
 	{
@@ -136,14 +145,31 @@ EPackageStoreEntryStatus FCookOnTheFlyPackageStoreBackend::CreatePackageStoreEnt
 	}
 }
 
-void FCookOnTheFlyPackageStoreBackend::AddPackages(TArray<FPackageStoreEntryResource> Entries, TArray<FPackageId> FailedPackageIds)
+void FCookOnTheFlyPackageStoreBackend::AddPackages(TArray<FPackageStoreEntryResource> Entries, TArray<FPackageId> FailedPackageIds,
+	TArray<TPair<FPackageId, FName>> PackageIdsAndNames)
 {
 	FWriteScopeLock WriteLock(EntriesLock);
 		
 	for (FPackageId FailedPackageId : FailedPackageIds)
 	{
-		UE_LOG(LogCookOnTheFly, Warning, TEXT("0x%llX [Failed]"), FailedPackageId.ValueForDebugging());
 		FEntryInfo& EntryInfo = PackageIdToEntryInfo.FindOrAdd(FailedPackageId, FEntryInfo());
+		if (EntryInfo.Status != EPackageStoreEntryStatus::Missing)
+		{
+			TPair<FPackageId, FName>* FailedPair = PackageIdsAndNames.FindByPredicate(
+				[FailedPackageId](TPair<FPackageId, FName>& Pair) { return Pair.Key == FailedPackageId; });
+			FName FailedName;
+			if (FailedPair)
+			{
+				FailedName = FailedPair->Value;
+			}
+			else
+			{
+				FailedName = FName(TEXT("<Unknown>"));
+			}
+
+			UE_LOG(LogCookOnTheFly, Warning, TEXT("0x%llX [Failed]. Failed to cook package %s."),
+				FailedPackageId.ValueForDebugging(), *FailedName.ToString());
+		}
 		EntryInfo.Status = EPackageStoreEntryStatus::Missing;
 		PackageStats.Failed++;
 	}
@@ -180,7 +206,8 @@ void FCookOnTheFlyPackageStoreBackend::OnCookOnTheFlyMessage(const UE::Cook::FCo
 				PackagesCookedMessage.CookedPackages.Num(),
 				PackagesCookedMessage.FailedPackages.Num());
 
-			AddPackages(MoveTemp(PackagesCookedMessage.CookedPackages), MoveTemp(PackagesCookedMessage.FailedPackages));
+			AddPackages(MoveTemp(PackagesCookedMessage.CookedPackages), MoveTemp(PackagesCookedMessage.FailedPackages),
+				TArray<TPair<FPackageId, FName>>());
 
 			if (Context)
 			{
diff --git a/Engine/Source/Runtime/StorageServerClient/Private/CookOnTheFlyPackageStore.h b/Engine/Source/Runtime/StorageServerClient/Private/CookOnTheFlyPackageStore.h
--- a/Engine/Source/Runtime/StorageServerClient/Private/CookOnTheFlyPackageStore.h
+++ b/Engine/Source/Runtime/StorageServerClient/Private/CookOnTheFlyPackageStore.h
@@ -38,7 +38,8 @@ public:
 	virtual void EndRead() override;
 
 	bool DoesPackageExist(FPackageId PackageId);
-	virtual EPackageStoreEntryStatus GetPackageStoreEntry(FPackageId PackageId, FPackageStoreEntry& OutPackageStoreEntry) override;
+	virtual EPackageStoreEntryStatus GetPackageStoreEntry(FPackageId PackageId, FName PackageName,
+		FPackageStoreEntry& OutPackageStoreEntry) override;
 	
 	virtual bool GetPackageRedirectInfo(FPackageId PackageId, FName& OutSourcePackageName, FPackageId& OutRedirectedToPackageId) override
 	{
@@ -46,9 +47,10 @@ public:
 	}
 	
 private:
-	void SendCookRequest(TArray<FPackageId> PackageIds);
+	void SendCookRequest(TArray<TPair<FPackageId, FName>> PackageIds);
 	EPackageStoreEntryStatus CreatePackageStoreEntry(const FEntryInfo& EntryInfo, FPackageStoreEntry& OutPackageStoreEntry);
-	void AddPackages(TArray<FPackageStoreEntryResource> Entries, TArray<FPackageId> FailedPackageIds);
+	void AddPackages(TArray<FPackageStoreEntryResource> Entries, TArray<FPackageId> FailedPackageIds,
+		TArray<TPair<FPackageId, FName>> PackageIdsAndNames);
 	void OnCookOnTheFlyMessage(const UE::Cook::FCookOnTheFlyMessage& Message);
 	void CheckActivity();
 
@@ -57,7 +59,7 @@ private:
 	FRWLock EntriesLock;
 	TMap<FPackageId, FEntryInfo> PackageIdToEntryInfo;
 	TChunkedArray<FPackageStoreEntryResource> PackageEntries;
-	TArray<FPackageId> RequestedPackageIds;
+	TArray<TPair<FPackageId, FName>> RequestedPackageIds;
 	FPackageStats PackageStats;
 
 	const double MaxInactivityTime = 20;
diff --git a/Engine/Source/Runtime/StorageServerClient/Private/StorageServerPackageStore.cpp b/Engine/Source/Runtime/StorageServerClient/Private/StorageServerPackageStore.cpp
--- a/Engine/Source/Runtime/StorageServerClient/Private/StorageServerPackageStore.cpp
+++ b/Engine/Source/Runtime/StorageServerClient/Private/StorageServerPackageStore.cpp
@@ -17,7 +17,8 @@ FStorageServerPackageStoreBackend::FStorageServerPackageStoreBackend(FStorageSer
 		});
 }
 
-EPackageStoreEntryStatus FStorageServerPackageStoreBackend::GetPackageStoreEntry(FPackageId PackageId, FPackageStoreEntry& OutPackageStoreEntry)
+EPackageStoreEntryStatus FStorageServerPackageStoreBackend::GetPackageStoreEntry(FPackageId PackageId, FName PackageName,
+	FPackageStoreEntry& OutPackageStoreEntry)
 {
 	const FStoreEntry* FindEntry = StoreEntriesMap.Find(PackageId);
 	if (FindEntry)
diff --git a/Engine/Source/Runtime/StorageServerClient/Private/StorageServerPackageStore.h b/Engine/Source/Runtime/StorageServerClient/Private/StorageServerPackageStore.h
--- a/Engine/Source/Runtime/StorageServerClient/Private/StorageServerPackageStore.h
+++ b/Engine/Source/Runtime/StorageServerClient/Private/StorageServerPackageStore.h
@@ -28,7 +28,8 @@ public:
 	{
 	}
 
-	virtual EPackageStoreEntryStatus GetPackageStoreEntry(FPackageId PackageIde, FPackageStoreEntry& OutPackageStoreEntry) override;
+	virtual EPackageStoreEntryStatus GetPackageStoreEntry(FPackageId PackageIde, FName PackageName,
+		FPackageStoreEntry& OutPackageStoreEntry) override;
 	
 	virtual bool GetPackageRedirectInfo(FPackageId PackageId, FName& OutSourcePackageName, FPackageId& OutRedirectedToPackageId) override
 	{
