Don't break in the handshake if more storage server connection hosts are available to try.

diff --git a/Engine/Source/Runtime/StorageServerClient/Private/StorageServerConnection.cpp b/Engine/Source/Runtime/StorageServerClient/Private/StorageServerConnection.cpp
--- a/Engine/Source/Runtime/StorageServerClient/Private/StorageServerConnection.cpp
+++ b/Engine/Source/Runtime/StorageServerClient/Private/StorageServerConnection.cpp
@@ -83,7 +83,7 @@ FStorageServerRequest::FStorageServerRequest(FAnsiStringView Verb, FAnsiStringVi
 		<< "Accept: " << GetMimeTypeString(Accept) << "\r\n";
 }
 
-FSocket* FStorageServerRequest::Send(FStorageServerConnection& Owner)
+FSocket* FStorageServerRequest::Send(FStorageServerConnection& Owner, bool bLogOnError)
 {
 	if (BodyBuffer.Num())
 	{
@@ -137,7 +137,10 @@ FSocket* FStorageServerRequest::Send(FStorageServerConnection& Owner)
 		}
 		return Socket;
 	}
-	UE_LOG(LogStorageServerConnection, Fatal, TEXT("Failed sending request to storage server."));
+	if (bLogOnError)
+	{
+		UE_LOG(LogStorageServerConnection, Fatal, TEXT("Failed sending request to storage server."));
+	}
 	return nullptr;
 }
 
@@ -517,7 +520,7 @@ int32 FStorageServerConnection::HandshakeRequest(TArrayView<const TSharedPtr<FIn
 		ReleaseSocket(ConnectSocket, true);
 
 		FStorageServerRequest Request("GET", *ResourceBuilder, Hostname);
-		if (FSocket* Socket = Request.Send(*this))
+		if (FSocket* Socket = Request.Send(*this, false))
 		{
 			FStorageServerResponse Response(*this, *Socket);
 
@@ -537,11 +540,17 @@ int32 FStorageServerConnection::HandshakeRequest(TArrayView<const TSharedPtr<FIn
 				UE_LOG(LogStorageServerConnection, Fatal, TEXT("Failed to handshake with Zen at %s. '%s'"), *ServerAddr->ToString(true), *Response.GetErrorMessage());
 			}
 		}
+		else
+		{
+			UE_LOG(LogStorageServerConnection, Warning, TEXT("Failed to send handshake request to Zen at %s."), *ServerAddr->ToString(true));
+		}
 	}
 
+	UE_LOG(LogStorageServerConnection, Fatal, TEXT("Failed to handshake with Zen at any of host addresses."));
+
 	Hostname.Reset();
 	ServerAddr.Reset();
-	
+
 	return -1;
 }
 
diff --git a/Engine/Source/Runtime/StorageServerClient/Public/StorageServerConnection.h b/Engine/Source/Runtime/StorageServerClient/Public/StorageServerConnection.h
--- a/Engine/Source/Runtime/StorageServerClient/Public/StorageServerConnection.h
+++ b/Engine/Source/Runtime/StorageServerClient/Public/StorageServerConnection.h
@@ -82,7 +82,7 @@ protected:
 		FAnsiStringView Hostname,
 		EStorageServerContentType Accept = EStorageServerContentType::Binary);
 
-	FSocket* Send(FStorageServerConnection& Owner);
+	FSocket* Send(FStorageServerConnection& Owner, bool bLogOnError = true);
 	virtual void Serialize(void* V, int64 Length) override;
 
 	EStorageServerContentType AcceptType;
