Change naming of .projectstore to ue.projectstore to avoid having hidden by default status on Mac/Linux.

diff --git a/Engine/Source/Developer/IoStoreUtilities/Private/ZenStoreWriter.cpp b/Engine/Source/Developer/IoStoreUtilities/Private/ZenStoreWriter.cpp
--- a/Engine/Source/Developer/IoStoreUtilities/Private/ZenStoreWriter.cpp
+++ b/Engine/Source/Developer/IoStoreUtilities/Private/ZenStoreWriter.cpp
@@ -415,7 +415,7 @@ void FZenStoreWriter::Initialize(const FCookInfo& Info)
 			IFileManager::Get().DeleteDirectory(*OutputPath, bRequireExists, bTree);
 		}
 
-		FString OplogLifetimeMarkerPath = OutputPath / TEXT(".projectstore");
+		FString OplogLifetimeMarkerPath = OutputPath / TEXT("ue.projectstore");
 		TUniquePtr<FArchive> OplogMarker(IFileManager::Get().CreateFileWriter(*OplogLifetimeMarkerPath));
 
 		bool bOplogEstablished = HttpClient->TryCreateOplog(ProjectId, OplogId, OplogLifetimeMarkerPath, CleanBuild);
diff --git a/Engine/Source/Programs/AutomationTool/BuildGraph/Tasks/ZenExportSnapshotTask.cs b/Engine/Source/Programs/AutomationTool/BuildGraph/Tasks/ZenExportSnapshotTask.cs
--- a/Engine/Source/Programs/AutomationTool/BuildGraph/Tasks/ZenExportSnapshotTask.cs
+++ b/Engine/Source/Programs/AutomationTool/BuildGraph/Tasks/ZenExportSnapshotTask.cs
@@ -112,7 +112,7 @@ namespace AutomationTool.Tasks
 		public string DestinationFileName;
 
 		/// <summary>
-		/// Optional. Where to look for the .projectstore
+		/// Optional. Where to look for the ue.projectstore
 		/// The pattern {Platform} can be used for exporting multiple platforms at once.
 		/// </summary>
 		[TaskParameter(Optional = true)]
@@ -357,7 +357,7 @@ namespace AutomationTool.Tasks
 					throw new AutomationException("Cook output directory not found ({0})", PlatformCookedDirectory.FullName);
 				}
 
-				FileReference ProjectStoreFile = FileReference.Combine(PlatformCookedDirectory, ".projectstore");
+				FileReference ProjectStoreFile = FileReference.Combine(PlatformCookedDirectory, "ue.projectstore");
 				if (!FileReference.Exists(ProjectStoreFile))
 				{
 					continue;
diff --git a/Engine/Source/Programs/AutomationTool/BuildGraph/Tasks/ZenImportOplogTask.cs b/Engine/Source/Programs/AutomationTool/BuildGraph/Tasks/ZenImportOplogTask.cs
--- a/Engine/Source/Programs/AutomationTool/BuildGraph/Tasks/ZenImportOplogTask.cs
+++ b/Engine/Source/Programs/AutomationTool/BuildGraph/Tasks/ZenImportOplogTask.cs
@@ -219,7 +219,7 @@ namespace AutomationTool.Tasks
 			{
 				DirectoryReference.CreateDirectory(PlatformCookedDirectory);
 			}
-			FileReference ProjectStoreFile = FileReference.Combine(PlatformCookedDirectory, ".projectstore");
+			FileReference ProjectStoreFile = FileReference.Combine(PlatformCookedDirectory, "ue.projectstore");
 
 			CbWriter Writer = new CbWriter();
 			Writer.BeginObject();
diff --git a/Engine/Source/Programs/AutomationTool/Scripts/RunProjectCommand.Automation.cs b/Engine/Source/Programs/AutomationTool/Scripts/RunProjectCommand.Automation.cs
--- a/Engine/Source/Programs/AutomationTool/Scripts/RunProjectCommand.Automation.cs
+++ b/Engine/Source/Programs/AutomationTool/Scripts/RunProjectCommand.Automation.cs
@@ -816,7 +816,7 @@ namespace AutomationScripts
 				return null;
 			}
 			// Check for stage with zenstore without PAK?
-			FileReference PackageStoreManifestFile = FileReference.Combine(ProjectStoreDir, ".projectstore");
+			FileReference PackageStoreManifestFile = FileReference.Combine(ProjectStoreDir, "ue.projectstore");
 			System.IO.FileInfo PackageStoreManifestFileInfo = PackageStoreManifestFile.ToFileInfo();
 			if (PackageStoreManifestFileInfo.Exists)
 			{
diff --git a/Engine/Source/Runtime/StorageServerClient/Private/StorageServerPlatformFile.cpp b/Engine/Source/Runtime/StorageServerClient/Private/StorageServerPlatformFile.cpp
--- a/Engine/Source/Runtime/StorageServerClient/Private/StorageServerPlatformFile.cpp
+++ b/Engine/Source/Runtime/StorageServerClient/Private/StorageServerPlatformFile.cpp
@@ -280,7 +280,7 @@ TUniquePtr<FArchive> FStorageServerPlatformFile::TryFindProjectStoreMarkerFile(I
 
 	for (const FString& ProjectStorePath : PotentialProjectStorePaths)
 	{
-		FString ProjectMarkerPath = ProjectStorePath / TEXT(".projectstore");
+		FString ProjectMarkerPath = ProjectStorePath / TEXT("ue.projectstore");
 		if (IFileHandle* ProjectStoreMarkerHandle = Inner->OpenRead(*ProjectMarkerPath); ProjectStoreMarkerHandle != nullptr)
 		{
 			UE_LOG(LogStorageServerPlatformFile, Display, TEXT("Found '%s'"), *ProjectMarkerPath);
@@ -325,7 +325,7 @@ bool FStorageServerPlatformFile::ShouldBeUsed(IPlatformFile* Inner, const TCHAR*
 			}
 
 			HostPort = ZenServerField["hostport"].AsUInt16(HostPort);
-			UE_LOG(LogStorageServerPlatformFile, Display, TEXT("Using connection settings from .projectstore: HostAddrs='%s' and HostPort='%d'"), *FString::Join(HostAddrs, TEXT("+")), HostPort);
+			UE_LOG(LogStorageServerPlatformFile, Display, TEXT("Using connection settings from ue.projectstore: HostAddrs='%s' and HostPort='%d'"), *FString::Join(HostAddrs, TEXT("+")), HostPort);
 		}
 	}
 
@@ -361,7 +361,7 @@ bool FStorageServerPlatformFile::Initialize(IPlatformFile* Inner, const TCHAR* C
 			{
 				ServerProject = FString(ZenServerField["projectid"].AsString());
 				ServerPlatform = FString(ZenServerField["oplogid"].AsString());
-				UE_LOG(LogStorageServerPlatformFile, Display, TEXT("Using settings from .projectstore: ServerProject='%s' and ServerPlatform='%s'"), *ServerProject, *ServerPlatform);
+				UE_LOG(LogStorageServerPlatformFile, Display, TEXT("Using settings from ue.projectstore: ServerProject='%s' and ServerPlatform='%s'"), *ServerProject, *ServerPlatform);
 			}
 		}
 	
