Improve the user experience when a game runtime is trying to connect to zenserver for network data stream but isn't able to make the connection.  This ensures we don't just end up at a halt on a fatal error, but give user-friendly messaging, with clarificitons on what actions the user can take.

diff --git a/Engine/Source/Runtime/StorageServerClient/Private/StorageServerConnection.cpp b/Engine/Source/Runtime/StorageServerClient/Private/StorageServerConnection.cpp
--- a/Engine/Source/Runtime/StorageServerClient/Private/StorageServerConnection.cpp
+++ b/Engine/Source/Runtime/StorageServerClient/Private/StorageServerConnection.cpp
@@ -546,7 +546,7 @@ int32 FStorageServerConnection::HandshakeRequest(TArrayView<const TSharedPtr<FIn
 		}
 	}
 
-	UE_LOG(LogStorageServerConnection, Fatal, TEXT("Failed to handshake with Zen at any of host addresses."));
+	UE_LOG(LogStorageServerConnection, Error, TEXT("Failed to handshake with Zen at any of host addresses."));
 
 	Hostname.Reset();
 	ServerAddr.Reset();
diff --git a/Engine/Source/Runtime/StorageServerClient/Private/StorageServerPlatformFile.cpp b/Engine/Source/Runtime/StorageServerClient/Private/StorageServerPlatformFile.cpp
--- a/Engine/Source/Runtime/StorageServerClient/Private/StorageServerPlatformFile.cpp
+++ b/Engine/Source/Runtime/StorageServerClient/Private/StorageServerPlatformFile.cpp
@@ -423,13 +423,17 @@ void FStorageServerPlatformFile::InitializeAfterProjectFilePath()
 			FText FailedConnectionTitle = NSLOCTEXT("StorageServer", "StorageServer_ConnectFailedTitle", "Failed to connect");
 			FText FailedConnectionText = FText::Format(NSLOCTEXT("StorageServer", "StorageServer_ConnectFailedText",
 				"Network data streaming failed to connect to any of the following data sources:\n\n{0}\n\n"
-				"This can be due to the sources being offline, invalid addresses, firewall blocking, or the sources being on a different network from this device. "
-				"If these issues can't be addressed, you can use an installed build without network data streaming. This process will now exit."),
+				"This can be due to the sources being offline, the Unreal Zen Storage process not currently running, "
+				"invalid addresses, firewall blocking, or the sources being on a different network from this device. "
+				"Please verify that your Unreal Zen Storage process is running using the ZenDashboard utility. "
+				"If these issues can't be addressed, you can use an installed build without network data streaming by "
+				"building with the '-pak' argument. This process will now exit."),
 				FText::FromString(FString::Join(HostAddrs, TEXT("\n"))));
 			FPlatformMisc::MessageBoxExt(EAppMsgType::Ok, *FailedConnectionText.ToString(), *FailedConnectionTitle.ToString());
 		}
 
-		UE_LOG(LogStorageServerPlatformFile, Fatal, TEXT("Failed to initialize connection to %s"), *FString::Join(HostAddrs, TEXT("\n")));
+		UE_LOG(LogStorageServerPlatformFile, Error, TEXT("Failed to initialize connection to %s"), *FString::Join(HostAddrs, TEXT("\n")));
+		FPlatformMisc::RequestExit(true);
 	}
 }
 
