- Add time scope for InitializeAfterProjectFilePath that can take more than 5 seconds during handshake
- Sort host address by loopback addresses (IPV6 and IPV4) and local host subnet in order to improve the chances of having a successfull connection at the first attempt and avoid possible timeouts

diff --git a/Engine/Source/Runtime/Launch/Private/LaunchEngineLoop.cpp b/Engine/Source/Runtime/Launch/Private/LaunchEngineLoop.cpp
--- a/Engine/Source/Runtime/Launch/Private/LaunchEngineLoop.cpp
+++ b/Engine/Source/Runtime/Launch/Private/LaunchEngineLoop.cpp
@@ -2564,9 +2564,13 @@ int32 FEngineLoop::PreInitPreStartupScreen(const TCHAR* CmdLine)
 
 	// Initialize platform file with knowledge of the project file path before fixing the casing
 	IPlatformFile* CurrentPlatformFile = &FPlatformFileManager::Get().GetPlatformFile();
-	for (IPlatformFile* PlatformFileChainElement = CurrentPlatformFile; PlatformFileChainElement; PlatformFileChainElement = PlatformFileChainElement->GetLowerLevel())
+
 	{
-		PlatformFileChainElement->InitializeAfterProjectFilePath();
+		SCOPED_BOOT_TIMING("PlatformFileChainElement->InitializeAfterProjectFilePath");
+		for (IPlatformFile* PlatformFileChainElement = CurrentPlatformFile; PlatformFileChainElement; PlatformFileChainElement = PlatformFileChainElement->GetLowerLevel())
+		{
+			PlatformFileChainElement->InitializeAfterProjectFilePath();
+		}
 	}
 
 #if !IS_PROGRAM
diff --git a/Engine/Source/Runtime/StorageServerClient/Private/StorageServerConnection.cpp b/Engine/Source/Runtime/StorageServerClient/Private/StorageServerConnection.cpp
--- a/Engine/Source/Runtime/StorageServerClient/Private/StorageServerConnection.cpp
+++ b/Engine/Source/Runtime/StorageServerClient/Private/StorageServerConnection.cpp
@@ -6,6 +6,7 @@
 #include "IPAddress.h"
 #include "Misc/App.h"
 #include "Misc/ScopeLock.h"
+#include "Misc/OutputDeviceRedirector.h"
 #include "Misc/StringBuilder.h"
 #include "Serialization/CompactBinary.h"
 #include "Serialization/CompactBinarySerialization.h"
@@ -496,12 +497,99 @@ bool FStorageServerConnection::Initialize(TArrayView<const FString> InHostAddres
 	return true;
 }
 
+void FStorageServerConnection::SortHostAddressesByLocalSubnet(TArrayView<const TSharedPtr<FInternetAddr>> HostAddresses, TArray<TSharedPtr<FInternetAddr>>& SortedHostAddresses)
+{
+	//no sorting needed cases
+	if (HostAddresses.Num() == 0)
+		return;
+
+	if (HostAddresses.Num() == 1)
+	{
+		SortedHostAddresses.Push(HostAddresses[0]);
+		return;
+	}
+
+	//Sorting logic:
+	//1 only on desktop, if it's an IPV6 address loopback (ends with ":1")
+	//2 only on desktop, if it's and IPV4 address loopback (starts with "127.0.0")
+	//3 if the host IPV4 subnet match the client subnet (xxx.xxx.xxx)
+	//4 remaining addresses
+	bool bCanBindAll = false;
+	bool bAppendPort = false;
+	TSharedPtr<FInternetAddr> localAddr = SocketSubsystem.GetLocalHostAddr(*GLog, bCanBindAll);
+	FString localAddrStringSubnet = localAddr->ToString(bAppendPort);
+
+	int32 localLastDotPos = INDEX_NONE;
+	if (localAddrStringSubnet.FindLastChar(TEXT('.'), localLastDotPos))
+	{
+		localAddrStringSubnet = localAddrStringSubnet.LeftChop(localAddrStringSubnet.Len() - localLastDotPos);
+	}
+
+	TArray<TSharedPtr<FInternetAddr>> IPV6Loopback;
+	TArray<TSharedPtr<FInternetAddr>> IPV4Loopback;
+	TArray<TSharedPtr<FInternetAddr>> RegularAddresses;
+
+	for (const TSharedPtr<FInternetAddr>& Addr : HostAddresses)
+	{
+		FString tempAddrStringSubnet = Addr->ToString(bAppendPort);
+
+#if PLATFORM_DESKTOP
+		if (Addr->GetProtocolType() == FNetworkProtocolTypes::IPv6)
+		{
+			if (tempAddrStringSubnet.EndsWith(":1"))
+			{
+				IPV6Loopback.Push(Addr);
+				continue;
+			}
+		}
+		else
+		{
+			if (tempAddrStringSubnet.StartsWith("127.0.0."))
+			{
+				IPV4Loopback.Push(Addr);
+				continue;
+			}
+		}
+#endif
+		int32 LastDotPos = INDEX_NONE;
+		if (tempAddrStringSubnet.FindLastChar(TEXT('.'), LastDotPos))
+		{
+			tempAddrStringSubnet = tempAddrStringSubnet.LeftChop(tempAddrStringSubnet.Len() - LastDotPos);
+		}
+
+		if (localAddrStringSubnet.Equals(tempAddrStringSubnet))
+			RegularAddresses.Insert(Addr, 0);
+		else
+			RegularAddresses.Push(Addr);
+	}
+
+
+	for (const TSharedPtr<FInternetAddr>& Addrv6lb : IPV6Loopback)
+	{
+		SortedHostAddresses.Push(Addrv6lb);
+	}
+
+	for (const TSharedPtr<FInternetAddr>& Addrv4lb : IPV4Loopback)
+	{
+		SortedHostAddresses.Push(Addrv4lb);
+	}
+
+	for (const TSharedPtr<FInternetAddr>& RegularAddr : RegularAddresses)
+	{
+		SortedHostAddresses.Push(RegularAddr);
+	}
+}
+
 int32 FStorageServerConnection::HandshakeRequest(TArrayView<const TSharedPtr<FInternetAddr>> HostAddresses)
 {
 	TAnsiStringBuilder<256> ResourceBuilder;
 	ResourceBuilder.Append(OplogPath);
 
-	for (const TSharedPtr<FInternetAddr>& Addr : HostAddresses)
+	// We sort the host address trying to match the local subnet in order to improve the chances of a successful connection at the first attempt.
+	TArray<TSharedPtr<FInternetAddr>> SortedAddresses;
+	SortHostAddressesByLocalSubnet(HostAddresses, SortedAddresses);
+
+	for (TSharedPtr<FInternetAddr>& Addr : SortedAddresses)
 	{
 		Hostname.Reset();
 		Hostname.Append(TCHAR_TO_ANSI(*Addr->ToString(false)));
diff --git a/Engine/Source/Runtime/StorageServerClient/Public/StorageServerConnection.h b/Engine/Source/Runtime/StorageServerClient/Public/StorageServerConnection.h
--- a/Engine/Source/Runtime/StorageServerClient/Public/StorageServerConnection.h
+++ b/Engine/Source/Runtime/StorageServerClient/Public/StorageServerConnection.h
@@ -195,6 +195,7 @@ private:
 	friend FStorageServerResponse;
 	friend FStorageServerChunkBatchRequest;
 
+	void SortHostAddressesByLocalSubnet(TArrayView<const TSharedPtr<FInternetAddr>> HostAddresses, TArray<TSharedPtr<FInternetAddr>>& SortedHostAddresses);
 	int32 HandshakeRequest(TArrayView<const TSharedPtr<FInternetAddr>> HostAddresses);
 	FSocket* AcquireSocketFromPool();
 	FSocket* AcquireNewSocket(float TimeoutSeconds = -1.0f);
