Convert ue.projectstore file to json.

diff --git a/Engine/Source/Developer/IoStoreUtilities/Private/ZenStoreWriter.cpp b/Engine/Source/Developer/IoStoreUtilities/Private/ZenStoreWriter.cpp
--- a/Engine/Source/Developer/IoStoreUtilities/Private/ZenStoreWriter.cpp
+++ b/Engine/Source/Developer/IoStoreUtilities/Private/ZenStoreWriter.cpp
@@ -1,35 +1,38 @@
 // Copyright Epic Games, Inc. All Rights Reserved.
 
 #include "ZenStoreWriter.h"
-#include "ZenStoreHttpClient.h"
-#include "ZenFileSystemManifest.h"
-#include "PackageStoreOptimizer.h"
+
 #include "Algo/BinarySearch.h"
 #include "Algo/IsSorted.h"
 #include "Algo/Sort.h"
 #include "AssetRegistry/AssetRegistryState.h"
 #include "Async/Async.h"
 #include "Containers/Queue.h"
-#include "Serialization/ArrayReader.h"
-#include "Serialization/CompactBinaryContainerSerialization.h"
-#include "Serialization/CompactBinaryPackage.h"
-#include "Serialization/CompactBinaryWriter.h"
-#include "Serialization/CompactBinarySerialization.h"
-#include "Serialization/LargeMemoryWriter.h" 
-#include "SocketSubsystem.h"
-#include "IPAddress.h"
+#include "GenericPlatform/GenericPlatformFile.h"
 #include "HAL/FileManager.h"
 #include "HAL/PlatformFileManager.h"
+#include "Interfaces/ITargetPlatform.h"
 #include "IO/IoDispatcher.h"
+#include "IPAddress.h"
 #include "Misc/App.h"
 #include "Misc/CommandLine.h"
 #include "Misc/FileHelper.h"
-#include "Misc/StringBuilder.h"
-#include "Interfaces/ITargetPlatform.h"
 #include "Misc/Paths.h"
-#include "GenericPlatform/GenericPlatformFile.h"
-#include "UObject/SavePackage.h"
 #include "Misc/PathViews.h"
+#include "Misc/StringBuilder.h"
+#include "PackageStoreOptimizer.h"
+#include "Serialization/ArrayReader.h"
+#include "Serialization/CompactBinaryContainerSerialization.h"
+#include "Serialization/CompactBinaryPackage.h"
+#include "Serialization/CompactBinarySerialization.h"
+#include "Serialization/CompactBinaryWriter.h"
+#include "Serialization/JsonSerializer.h"
+#include "Serialization/JsonWriter.h"
+#include "Serialization/LargeMemoryWriter.h" 
+#include "SocketSubsystem.h"
+#include "UObject/SavePackage.h"
+#include "ZenFileSystemManifest.h"
+#include "ZenStoreHttpClient.h"
 
 DEFINE_LOG_CATEGORY_STATIC(LogZenStoreWriter, Log, All);
 
@@ -422,15 +425,15 @@ void FZenStoreWriter::Initialize(const FCookInfo& Info)
 
 		if (bOplogEstablished && OplogMarker)
 		{
-			FCbWriter ManifestWriter;
-			ManifestWriter.BeginObject();
-			ManifestWriter.BeginObject("zenserver");
 			bool IsRunningLocally = false;
 #if UE_WITH_ZEN
 			IsRunningLocally = HttpClient->GetZenServiceInstance().IsServiceRunningLocally();
 #endif
-			ManifestWriter << "islocalhost" << IsRunningLocally;
-			ManifestWriter << "hostname" << HttpClient->GetHostName();
+			TSharedRef<TJsonWriter<UTF8CHAR, TPrettyJsonPrintPolicy<UTF8CHAR>>> Writer = TJsonWriterFactory<UTF8CHAR, TPrettyJsonPrintPolicy<UTF8CHAR>>::Create(OplogMarker.Get());
+			Writer->WriteObjectStart();
+			Writer->WriteObjectStart(TEXT("zenserver"));
+			Writer->WriteValue(TEXT("islocalhost"), IsRunningLocally);
+			Writer->WriteValue(TEXT("hostname"), HttpClient->GetHostName());
 			if (IsRunningLocally)
 			{
 				ISocketSubsystem* SocketSubsystem = ISocketSubsystem::Get();
@@ -439,21 +442,21 @@ void FZenStoreWriter::Initialize(const FCookInfo& Info)
 					TArray<TSharedPtr<FInternetAddr>> Addresses;
 					if (SocketSubsystem->GetLocalAdapterAddresses(Addresses))
 					{
-						ManifestWriter.BeginArray("remotehostnames");
+						Writer->WriteArrayStart("remotehostnames");
 						for (const TSharedPtr<FInternetAddr>& Address : Addresses)
 						{
-							ManifestWriter << Address->ToString(false);
+							Writer->WriteValue(Address->ToString(false));
 						}
-						ManifestWriter.EndArray();
+						Writer->WriteArrayEnd();
 					}
 				}
 			}
-			ManifestWriter << "hostport" << HttpClient->GetPort();
-			ManifestWriter << "projectid" << ProjectId;
-			ManifestWriter << "oplogid" << OplogId;
-			ManifestWriter.EndObject();
-			ManifestWriter.EndObject();
-			SaveCompactBinary(*OplogMarker, ManifestWriter.Save());
+			Writer->WriteValue(TEXT("hostport"), HttpClient->GetPort());
+			Writer->WriteValue(TEXT("projectid"), ProjectId);
+			Writer->WriteValue(TEXT("oplogid"), OplogId);
+			Writer->WriteObjectEnd();
+			Writer->WriteObjectEnd();
+			Writer->Close();
 		}
 
 		OplogMarker.Reset();
diff --git a/Engine/Source/Programs/AutomationTool/BuildGraph/BuildGraph.Automation.csproj b/Engine/Source/Programs/AutomationTool/BuildGraph/BuildGraph.Automation.csproj
--- a/Engine/Source/Programs/AutomationTool/BuildGraph/BuildGraph.Automation.csproj
+++ b/Engine/Source/Programs/AutomationTool/BuildGraph/BuildGraph.Automation.csproj
@@ -55,6 +55,7 @@
     <ProjectReference Include="..\..\Shared\EpicGames.BuildGraph\EpicGames.BuildGraph.csproj" />
     <ProjectReference Include="..\..\Shared\EpicGames.Core\EpicGames.Core.csproj" PrivateAssets="All"><Private>false</Private></ProjectReference>
     <ProjectReference Include="..\..\Shared\EpicGames.Jupiter\EpicGames.Jupiter.csproj" PrivateAssets="All" />
+    <ProjectReference Include="..\..\Shared\EpicGames.ProjectStore\EpicGames.ProjectStore.csproj" PrivateAssets="All" />
     <ProjectReference Include="..\..\UnrealBuildTool\UnrealBuildTool.csproj" PrivateAssets="All"><Private>false</Private></ProjectReference>
     <ProjectReference Include="..\AutomationUtils\AutomationUtils.Automation.csproj" PrivateAssets="All"><Private>false</Private></ProjectReference>
   </ItemGroup>
diff --git a/Engine/Source/Programs/AutomationTool/BuildGraph/Tasks/ZenExportSnapshotTask.cs b/Engine/Source/Programs/AutomationTool/BuildGraph/Tasks/ZenExportSnapshotTask.cs
--- a/Engine/Source/Programs/AutomationTool/BuildGraph/Tasks/ZenExportSnapshotTask.cs
+++ b/Engine/Source/Programs/AutomationTool/BuildGraph/Tasks/ZenExportSnapshotTask.cs
@@ -14,6 +14,7 @@ using System.Threading.Tasks;
 using System.Xml;
 using EpicGames.Core;
 using EpicGames.Serialization;
+using EpicGames.ProjectStore;
 
 namespace AutomationTool.Tasks
 {
@@ -392,22 +393,15 @@ namespace AutomationTool.Tasks
 				}
 
 				FileReference ProjectStoreFile = FileReference.Combine(PlatformCookedDirectory, "ue.projectstore");
-				if (!FileReference.Exists(ProjectStoreFile))
-				{
-					continue;
-				}
-				
-				byte[] ProjectStoreData = File.ReadAllBytes(ProjectStoreFile.FullName);
-				CbObject ProjectStoreObject = new CbField(ProjectStoreData).AsObject();
-				CbObject ZenServerObject = ProjectStoreObject["zenserver"].AsObject();
-				if (ZenServerObject != CbObject.Empty)
+				ProjectStoreData? ParsedProjectStore = null;
+				if (TryLoadJson(ProjectStoreFile, out ParsedProjectStore) && (ParsedProjectStore != null) && (ParsedProjectStore.ZenServer != null))
 				{
 					ExportSourceData NewExportSource = new ExportSourceData();
-					NewExportSource.IsLocalHost = ZenServerObject["islocalhost"].AsBool();
-					NewExportSource.HostName = ZenServerObject["hostname"].AsString("localhost");
-					NewExportSource.HostPort = ZenServerObject["hostport"].AsInt16(8558);
-					NewExportSource.ProjectId = ZenServerObject["projectid"].AsString();
-					NewExportSource.OplogId = ZenServerObject["oplogid"].AsString();
+					NewExportSource.IsLocalHost = ParsedProjectStore.ZenServer.IsLocalHost;
+					NewExportSource.HostName = ParsedProjectStore.ZenServer.HostName;
+					NewExportSource.HostPort = ParsedProjectStore.ZenServer.HostPort;
+					NewExportSource.ProjectId = ParsedProjectStore.ZenServer.ProjectId;
+					NewExportSource.OplogId = ParsedProjectStore.ZenServer.OplogId;
 					NewExportSource.TargetPlatform = Platform;
 					NewExportSource.SnapshotBaseDescriptor = null;
 
diff --git a/Engine/Source/Programs/AutomationTool/BuildGraph/Tasks/ZenImportOplogTask.cs b/Engine/Source/Programs/AutomationTool/BuildGraph/Tasks/ZenImportOplogTask.cs
--- a/Engine/Source/Programs/AutomationTool/BuildGraph/Tasks/ZenImportOplogTask.cs
+++ b/Engine/Source/Programs/AutomationTool/BuildGraph/Tasks/ZenImportOplogTask.cs
@@ -1,16 +1,16 @@
 // Copyright Epic Games, Inc. All Rights Reserved.
 
-using AutomationTool;
-using EpicGames.BuildGraph;
+using EpicGames.Core;
+using EpicGames.ProjectStore;
+using Microsoft.Extensions.Logging;
 using System;
 using System.Collections.Generic;
+using System.IO;
 using System.Text;
+using System.Text.Json;
+using System.Text.Json.Serialization;
 using System.Threading.Tasks;
 using System.Xml;
-using EpicGames.Core;
-using UnrealBuildTool;
-using Microsoft.Extensions.Logging;
-using EpicGames.Serialization;
 
 namespace AutomationTool.Tasks
 {
@@ -219,22 +219,23 @@ namespace AutomationTool.Tasks
 			{
 				DirectoryReference.CreateDirectory(PlatformCookedDirectory);
 			}
-			FileReference ProjectStoreFile = FileReference.Combine(PlatformCookedDirectory, "ue.projectstore");
-
-			CbWriter Writer = new CbWriter();
-			Writer.BeginObject();
-			Writer.BeginObject("zenserver");
-
-			Writer.WriteBool("islocalhost", true);
-			Writer.WriteString("hostname", "localhost");
-			Writer.WriteInteger("hostport", 8558);
-			Writer.WriteString("projectid", Parameters.ProjectName);
-			Writer.WriteString("oplogid", Parameters.OplogName);
+			ProjectStoreData ProjectStore = new ProjectStoreData();
+			ProjectStore.ZenServer = new ZenServerStoreData
+			{
+				ProjectId = Parameters.ProjectName,
+				OplogId = Parameters.OplogName
+			};
 
-			Writer.EndObject();
-			Writer.EndObject();
+			JsonSerializerOptions SerializerOptions = new JsonSerializerOptions
+			{
+				AllowTrailingCommas = true,
+				ReadCommentHandling = JsonCommentHandling.Skip,
+				PropertyNameCaseInsensitive = true
+			};
+			SerializerOptions.Converters.Add(new JsonStringEnumConverter());
 
-			FileReference.WriteAllBytes(ProjectStoreFile, Writer.ToByteArray());
+			FileReference ProjectStoreFile = FileReference.Combine(PlatformCookedDirectory, "ue.projectstore");
+			File.WriteAllText(ProjectStoreFile.FullName, JsonSerializer.Serialize(ProjectStore, SerializerOptions), new UTF8Encoding(encoderShouldEmitUTF8Identifier: false));
 		}
 
 		private void ImportFromCloud(FileReference ZenExe)
diff --git a/Engine/Source/Programs/AutomationTool/Scripts/AutomationScripts.Automation.csproj b/Engine/Source/Programs/AutomationTool/Scripts/AutomationScripts.Automation.csproj
--- a/Engine/Source/Programs/AutomationTool/Scripts/AutomationScripts.Automation.csproj
+++ b/Engine/Source/Programs/AutomationTool/Scripts/AutomationScripts.Automation.csproj
@@ -49,6 +49,7 @@
   <ItemGroup>
     <ProjectReference Include="..\..\Shared\EpicGames.Core\EpicGames.Core.csproj" PrivateAssets="All"><Private>false</Private></ProjectReference>
     <ProjectReference Include="..\..\Shared\EpicGames.Jupiter\EpicGames.Jupiter.csproj" PrivateAssets="All" />
+    <ProjectReference Include="..\..\Shared\EpicGames.ProjectStore\EpicGames.ProjectStore.csproj" PrivateAssets="All" />
     <ProjectReference Include="..\..\UnrealBuildTool\UnrealBuildTool.csproj" PrivateAssets="All"><Private>false</Private></ProjectReference>
     <ProjectReference Include="..\AutomationUtils\AutomationUtils.Automation.csproj" PrivateAssets="All"><Private>false</Private></ProjectReference>
     <ProjectReference Include="..\Localization\Localization.Automation.csproj" PrivateAssets="All"><Private>false</Private></ProjectReference>
diff --git a/Engine/Source/Programs/Shared/EpicGames.ProjectStore/EpicGames.ProjectStore.csproj b/Engine/Source/Programs/Shared/EpicGames.ProjectStore/EpicGames.ProjectStore.csproj
new file mode 100644
--- /dev/null
+++ b/Engine/Source/Programs/Shared/EpicGames.ProjectStore/EpicGames.ProjectStore.csproj
@@ -0,0 +1,37 @@
+<Project Sdk="Microsoft.NET.Sdk">
+  <PropertyGroup>
+    <TargetFramework>net6.0</TargetFramework>
+    <Configuration Condition=" '$(Configuration)' == '' ">Development</Configuration>
+    <OutputType>Library</OutputType>
+    <Deterministic>true</Deterministic>
+    <RootNamespace>EpicGames.ProjectStore</RootNamespace>
+    <GenerateAssemblyInfo>false</GenerateAssemblyInfo>
+    <GenerateTargetFrameworkAttribute>false</GenerateTargetFrameworkAttribute>
+    <Configurations>Debug;Release;Development;Analyze</Configurations>
+    <Nullable>enable</Nullable>
+    <DebugType>pdbonly</DebugType>
+    <DebugType Condition="'$([System.Runtime.InteropServices.RuntimeInformation]::IsOSPlatform($([System.Runtime.InteropServices.OSPlatform]::Windows)))' == 'true' And '$([System.Runtime.InteropServices.RuntimeInformation]::OSArchitecture)' == 'Arm64'">portable</DebugType>
+    <RunAnalyzersDuringBuild>false</RunAnalyzersDuringBuild>
+  </PropertyGroup>
+  <PropertyGroup Condition=" '$(Configuration)|$(Platform)' == 'Debug|AnyCPU' ">
+    <DocumentationFile></DocumentationFile>
+  </PropertyGroup>
+  <PropertyGroup Condition="'$(Configuration)|$(Platform)' == 'Development|AnyCPU'">
+    <DocumentationFile></DocumentationFile>
+    <Optimize>true</Optimize>
+  </PropertyGroup>
+  <PropertyGroup Condition="'$(Configuration)|$(Platform)' == 'Release|AnyCPU'">
+    <DocumentationFile></DocumentationFile>
+    <Optimize>true</Optimize>
+  </PropertyGroup>
+  <PropertyGroup Condition=" '$(Configuration)|$(Platform)' == 'Analyze|AnyCPU' ">
+    <DocumentationFile></DocumentationFile>
+    <RunAnalyzersDuringBuild>True</RunAnalyzersDuringBuild>
+    <GenerateDocumentationFile>True</GenerateDocumentationFile>
+  </PropertyGroup>
+  <ItemGroup>
+    <Compile Include="..\..\Shared\MetaData.cs">
+      <Link>Properties\MetaData.cs</Link>
+    </Compile>
+  </ItemGroup>
+</Project>
\ No newline at end of file
diff --git a/Engine/Source/Programs/Shared/EpicGames.ProjectStore/ProjectStoreTypes.cs b/Engine/Source/Programs/Shared/EpicGames.ProjectStore/ProjectStoreTypes.cs
new file mode 100644
--- /dev/null
+++ b/Engine/Source/Programs/Shared/EpicGames.ProjectStore/ProjectStoreTypes.cs
@@ -0,0 +1,47 @@
+// Copyright Epic Games, Inc. All Rights Reserved.
+
+using System.Collections.Generic;
+
+namespace EpicGames.ProjectStore
+{
+	/// <summary>
+	/// A specification of storage details with zenserver
+	/// </summary>
+	public class ZenServerStoreData
+	{
+		/// <summary>
+		/// Whether zenserver is running locally
+		/// </summary>
+		public bool IsLocalHost { get; set; } = true;
+		/// <summary>
+		/// The primary or local host name for zenserver
+		/// </summary>
+		public string? HostName { get; set; } = "localhost";
+		/// <summary>
+		/// The list of remote host names for zenserver
+		/// </summary>
+		public List<string> RemoteHostNames { get; set; } = new List<string>();
+		/// <summary>
+		/// The port for zenserver on the host
+		/// </summary>
+		public int HostPort { get; set; } = 8558;
+		/// <summary>
+		/// The project identifier for the stored data in zenserver
+		/// </summary>
+		public string? ProjectId { get; set; }
+		/// <summary>
+		/// The oplog identifier for the stored data in zenserver
+		/// </summary>
+		public string? OplogId { get; set; }
+	}
+	/// <summary>
+	/// A desctiptor of a project store
+	/// </summary>
+	public class ProjectStoreData
+	{
+		/// <summary>
+		/// The details for zenserver for this project store
+		/// </summary>
+		public ZenServerStoreData? ZenServer { get; set; }
+	}
+}
diff --git a/Engine/Source/Runtime/StorageServerClient/Private/StorageServerPlatformFile.cpp b/Engine/Source/Runtime/StorageServerClient/Private/StorageServerPlatformFile.cpp
--- a/Engine/Source/Runtime/StorageServerClient/Private/StorageServerPlatformFile.cpp
+++ b/Engine/Source/Runtime/StorageServerClient/Private/StorageServerPlatformFile.cpp
@@ -9,12 +9,15 @@
 #include "Misc/App.h"
 #include "Misc/CommandLine.h"
 #include "Misc/CoreDelegates.h"
+#include "Misc/FileHelper.h"
 #include "Misc/Paths.h"
 #include "Misc/ScopeRWLock.h"
 #include "Misc/StringBuilder.h"
 #include "Modules/ModuleManager.h"
 #include "Modules/ModuleManager.h"
 #include "Serialization/CompactBinarySerialization.h"
+#include "Serialization/JsonReader.h"
+#include "Serialization/JsonSerializer.h"
 #include "StorageServerConnection.h"
 #include "StorageServerIoDispatcherBackend.h"
 #include "StorageServerPackageStore.h"
@@ -305,27 +308,40 @@ bool FStorageServerPlatformFile::ShouldBeUsed(IPlatformFile* Inner, const TCHAR*
 	TUniquePtr<FArchive> ProjectStoreMarkerReader = TryFindProjectStoreMarkerFile(Inner);
 	if (ProjectStoreMarkerReader != nullptr)
 	{
-		FCbObject ProjectStoreObject = LoadCompactBinary(*ProjectStoreMarkerReader).AsObject();
-
-		if (FCbFieldView ZenServerField = ProjectStoreObject["zenserver"])
+		TSharedPtr<FJsonObject> ProjectStoreObject;
+		TSharedRef<TJsonReader<UTF8CHAR>> Reader = TJsonReaderFactory<UTF8CHAR>::Create(ProjectStoreMarkerReader.Get());
+		if (FJsonSerializer::Deserialize(Reader, ProjectStoreObject) && ProjectStoreObject.IsValid())
 		{
-#if PLATFORM_DESKTOP
-			if (FUtf8StringView HostName = ZenServerField["hostname"].AsString(); !HostName.IsEmpty())
+			const TSharedPtr<FJsonObject>* ZenServerObjectPtr = nullptr;
+			if (ProjectStoreObject->TryGetObjectField(TEXT("zenserver"), ZenServerObjectPtr) && (ZenServerObjectPtr != nullptr))
 			{
-				HostAddrs.Add(FString(HostName));
-			}
+				const TSharedPtr<FJsonObject>& ZenServerObject = *ZenServerObjectPtr;
+#if PLATFORM_DESKTOP
+				FString HostName;
+				if (ZenServerObject->TryGetStringField(TEXT("hostname"), HostName) && !HostName.IsEmpty())
+				{
+					HostAddrs.Add(HostName);
+				}
 #endif
-			FCbArrayView RemoteHostNames = ZenServerField["remotehostnames"].AsArrayView();
-			for (FCbFieldView RemoteHostName : RemoteHostNames)
-			{
-				if (FUtf8StringView HostName = RemoteHostName.AsString(); !HostName.IsEmpty())
+				const TArray<TSharedPtr<FJsonValue>>* RemoteHostNamesArrayPtr = nullptr;
+				if (ZenServerObject->TryGetArrayField(TEXT("remotehostnames"), RemoteHostNamesArrayPtr) && (RemoteHostNamesArrayPtr != nullptr))
 				{
-					HostAddrs.Add(FString(HostName));
+					for (TSharedPtr<FJsonValue> RemoteHostName : *RemoteHostNamesArrayPtr)
+					{
+						if (FString RemoteHostNameStr = RemoteHostName->AsString(); !RemoteHostNameStr.IsEmpty())
+						{
+							HostAddrs.Add(RemoteHostNameStr);
+						}
+					}
 				}
-			}
 
-			HostPort = ZenServerField["hostport"].AsUInt16(HostPort);
-			UE_LOG(LogStorageServerPlatformFile, Display, TEXT("Using connection settings from ue.projectstore: HostAddrs='%s' and HostPort='%d'"), *FString::Join(HostAddrs, TEXT("+")), HostPort);
+				uint16 SerializedHostPort = 0;
+				if (ZenServerObject->TryGetNumberField(TEXT("hostport"), SerializedHostPort) && (SerializedHostPort != 0))
+				{
+					HostPort = SerializedHostPort;
+				}
+				UE_LOG(LogStorageServerPlatformFile, Display, TEXT("Using connection settings from ue.projectstore: HostAddrs='%s' and HostPort='%d'"), *FString::Join(HostAddrs, TEXT("+")), HostPort);
+			}
 		}
 	}
 
@@ -355,13 +371,18 @@ bool FStorageServerPlatformFile::Initialize(IPlatformFile* Inner, const TCHAR* C
 		TUniquePtr<FArchive> ProjectStoreMarkerReader = TryFindProjectStoreMarkerFile(Inner);
 		if (ProjectStoreMarkerReader != nullptr)
 		{
-			FCbObject ProjectStoreObject = LoadCompactBinary(*ProjectStoreMarkerReader).AsObject();
-
-			if (FCbFieldView ZenServerField = ProjectStoreObject["zenserver"])
+			TSharedPtr<FJsonObject> ProjectStoreObject;
+			TSharedRef<TJsonReader<UTF8CHAR>> Reader = TJsonReaderFactory<UTF8CHAR>::Create(ProjectStoreMarkerReader.Get());
+			if (FJsonSerializer::Deserialize(Reader, ProjectStoreObject) && ProjectStoreObject.IsValid())
 			{
-				ServerProject = FString(ZenServerField["projectid"].AsString());
-				ServerPlatform = FString(ZenServerField["oplogid"].AsString());
-				UE_LOG(LogStorageServerPlatformFile, Display, TEXT("Using settings from ue.projectstore: ServerProject='%s' and ServerPlatform='%s'"), *ServerProject, *ServerPlatform);
+				const TSharedPtr<FJsonObject>* ZenServerObjectPtr = nullptr;
+				if (ProjectStoreObject->TryGetObjectField(TEXT("zenserver"), ZenServerObjectPtr) && (ZenServerObjectPtr != nullptr))
+				{
+					const TSharedPtr<FJsonObject>& ZenServerObject = *ZenServerObjectPtr;
+					ServerProject = ZenServerObject->GetStringField(TEXT("projectid"));
+					ServerPlatform = ZenServerObject->GetStringField(TEXT("oplogid"));
+					UE_LOG(LogStorageServerPlatformFile, Display, TEXT("Using settings from ue.projectstore: ServerProject='%s' and ServerPlatform='%s'"), *ServerProject, *ServerPlatform);
+				}
 			}
 		}
 	
diff --git a/Engine/Source/Runtime/StorageServerClient/StorageServerClient.Build.cs b/Engine/Source/Runtime/StorageServerClient/StorageServerClient.Build.cs
--- a/Engine/Source/Runtime/StorageServerClient/StorageServerClient.Build.cs
+++ b/Engine/Source/Runtime/StorageServerClient/StorageServerClient.Build.cs
@@ -11,7 +11,8 @@ public class StorageServerClient : ModuleRules
                 "Core",
                 "CoreUObject",
                 "Sockets",
-				"CookOnTheFly"
+                "CookOnTheFly",
+                "Json"
             }
         );
     }
