Add support for compressed binary PackageStoreRequest

diff --git a/Engine/Source/Runtime/StorageServerClient/Private/StorageServerConnection.cpp b/Engine/Source/Runtime/StorageServerClient/Private/StorageServerConnection.cpp
--- a/Engine/Source/Runtime/StorageServerClient/Private/StorageServerConnection.cpp
+++ b/Engine/Source/Runtime/StorageServerClient/Private/StorageServerConnection.cpp
@@ -436,7 +436,37 @@ int64 FStorageServerResponse::SerializeChunkTo(FMutableMemoryView Memory, uint64
 
 FCbObject FStorageServerResponse::GetResponseObject()
 {
-	FCbField Payload = LoadCompactBinary(*this);
+	FCbField Payload;
+
+	if (ContentType == EStorageServerContentType::CompressedBinary)
+	{
+		TRACE_CPUPROFILER_EVENT_SCOPE(GetCompressedResponseObject);
+
+		TArray64<uint8> CompressedBuffer;
+		CompressedBuffer.Reset(ContentLength);
+		{
+			TRACE_CPUPROFILER_EVENT_SCOPE(SerializeCompressedData);
+			Serialize(CompressedBuffer.GetData(), ContentLength);
+		}
+
+		if (FCompressedBuffer Compressed = FCompressedBuffer::FromCompressed(FSharedBuffer::MakeView(CompressedBuffer.GetData(), ContentLength)))
+		{
+			FIoBuffer Decompressed(Compressed.GetRawSize());
+			if (FCompressedBufferReader(Compressed).TryDecompressTo(Decompressed.GetMutableView(), 0))
+			{
+				FBufferReader DecompressedAr(Decompressed.GetData(), Decompressed.GetSize(), false);
+				{
+					TRACE_CPUPROFILER_EVENT_SCOPE(LoadingUncompressedCompactBinary);
+					Payload = LoadCompactBinary(DecompressedAr);
+				}
+			}
+		}
+	}
+	else if (ContentType == EStorageServerContentType::CbObject)
+	{
+		Payload = LoadCompactBinary(*this);
+	}
+	
 	return Payload.AsObject();
 }
 
@@ -644,9 +674,11 @@ int32 FStorageServerConnection::HandshakeRequest(TArrayView<const TSharedPtr<FIn
 
 void FStorageServerConnection::PackageStoreRequest(TFunctionRef<void(FPackageStoreEntryResource&&)> Callback)
 {
+	TRACE_CPUPROFILER_EVENT_SCOPE(StorageServerPackageStoreRequest);
+
 	TAnsiStringBuilder<256> ResourceBuilder;
-	ResourceBuilder.Append(OplogPath).Append("/entries");
-	FStorageServerRequest Request("GET", *ResourceBuilder, Hostname, EStorageServerContentType::CbObject);
+	ResourceBuilder.Append(OplogPath).Append("/entries?fieldfilter=packagestoreentry");
+	FStorageServerRequest Request("GET", *ResourceBuilder, Hostname, EStorageServerContentType::CompressedBinary);
 	FSocket* Socket = Request.Send(*this);
 	if (!Socket)
 	{
@@ -657,6 +689,8 @@ void FStorageServerConnection::PackageStoreRequest(TFunctionRef<void(FPackageSto
 	if (Response.IsOk())
 	{
 		FCbObject ResponseObj = Response.GetResponseObject();
+
+		TRACE_CPUPROFILER_EVENT_SCOPE(StorageServerPackageStoreRequestParseEntries);
 		for (FCbField& OplogEntry : ResponseObj["entries"].AsArray())
 		{
 			FCbObject OplogObj = OplogEntry.AsObject();
