Zen transport plugin for staged builds serving using the platform protocol interface and FStorageServerConnection adaptation for platform backends

diff --git a/Engine/Source/Runtime/CookOnTheFly/Private/CookOnTheFlyModule.cpp b/Engine/Source/Runtime/CookOnTheFly/Private/CookOnTheFlyModule.cpp
--- a/Engine/Source/Runtime/CookOnTheFly/Private/CookOnTheFlyModule.cpp
+++ b/Engine/Source/Runtime/CookOnTheFly/Private/CookOnTheFlyModule.cpp
@@ -127,6 +127,11 @@ private:
 			ServerStartupWaitTime = FTimespan::FromSeconds(ServerWaitTimeInSeconds);
 		}
 
+		if (FParse::Param(FCommandLine::Get(), TEXT("ZenPlatformProtocol")))
+		{
+			HostIpList.Add("platform://");
+		}
+
 		return { HostIpList, ServerStartupWaitTime };
 	}
 };
diff --git a/Engine/Source/Runtime/CookOnTheFly/Private/PlatformTransport.cpp b/Engine/Source/Runtime/CookOnTheFly/Private/PlatformTransport.cpp
--- a/Engine/Source/Runtime/CookOnTheFly/Private/PlatformTransport.cpp
+++ b/Engine/Source/Runtime/CookOnTheFly/Private/PlatformTransport.cpp
@@ -24,7 +24,8 @@ namespace
 
 		virtual bool Receive(uint8* Results, int32 Size) const override
 		{
-			return (HostSocket->Receive(Results, Size) == IPlatformHostSocket::EResultNet::Ok);
+			uint64 BytesReceived;
+			return (HostSocket->Receive(Results, Size, BytesReceived) == IPlatformHostSocket::EResultNet::Ok);
 		}
 
 		virtual bool Send(const uint8* Buffer, int32 Size) const override
diff --git a/Engine/Source/Runtime/Core/Public/GenericPlatform/GenericPlatformHostSocket.h b/Engine/Source/Runtime/Core/Public/GenericPlatform/GenericPlatformHostSocket.h
--- a/Engine/Source/Runtime/Core/Public/GenericPlatform/GenericPlatformHostSocket.h
+++ b/Engine/Source/Runtime/Core/Public/GenericPlatform/GenericPlatformHostSocket.h
@@ -48,6 +48,15 @@ public:
 		Closed,			// Socket has already been closed and shouldn't be used.
 	};
 
+	/**
+	 * Mode of a socket read
+	 */
+	enum class EReceiveFlags : uint8
+	{
+		DontWait,		// Read as much there is on the wire up to the buffer size
+		WaitAll,		// Block read and wait until the buffer is filled
+	};
+
 public:
 
 	/**
@@ -61,13 +70,24 @@ public:
 
 	/**
 	 * Receive data from the connected host PC (blocking operation).
-	 * 
+	 *
 	 * @param Buffer         Data to be sent.
 	 * @param BytesToReceive The number of bytes to receive (Buffer has to be large enough).
 	 * @return               Status value indicating error or success.
 	 */
 	virtual EResultNet Receive(void* Buffer, uint64 BytesToReceive) = 0;
 
+	/**
+	 * Receive data from the connected host PC.
+	 * 
+	 * @param Buffer         Data to be sent.
+	 * @param BytesToReceive The number of bytes to receive (Buffer has to be large enough).
+	 * @param BytesReceived  Number of bytes that have been received (equals to BytesToReceive if ReadMode is EReceiveFlags::WaitAll)
+	 * @param ReadMode       DontWait if this call should return immediately with the data available and not wait for BytesToReceive number of bytes
+	 * @return               Status value indicating error or success.
+	 */
+	virtual EResultNet Receive(void* Buffer, uint64 BytesToReceive, uint64& BytesReceived, EReceiveFlags ReadMode = EReceiveFlags::WaitAll) = 0;
+
 	/**
 	 * Get the state of the socket (determines if the host pc is connected and communication is possible).
 	 */
diff --git a/Engine/Source/Runtime/StorageServerClient/Private/StoragePlatformConnectionBackend.cpp b/Engine/Source/Runtime/StorageServerClient/Private/StoragePlatformConnectionBackend.cpp
new file mode 100644
--- /dev/null
+++ b/Engine/Source/Runtime/StorageServerClient/Private/StoragePlatformConnectionBackend.cpp
@@ -0,0 +1,254 @@
+// Copyright Epic Games, Inc. All Rights Reserved.
+
+#include "StoragePlatformConnectionBackend.h"
+
+#include "Serialization/CompactBinary.h"
+#include "Serialization/CompactBinarySerialization.h"
+#include "HAL/PlatformProcess.h"
+
+DEFINE_LOG_CATEGORY_STATIC(LogStorageServerPlatformBackend, Log, All);
+
+FStorageConnectionPlatformSocket::FStorageConnectionPlatformSocket(IPlatformHostCommunication* HostCommunication, IPlatformHostSocketPtr InSocket, int32 InProtocolNumber)
+	: Communication(HostCommunication)
+	, Socket(InSocket)
+	, ConnectionBuffer(1024 * 64)
+	, ProtocolNumber(InProtocolNumber)
+{
+}
+
+FStorageConnectionPlatformSocket::~FStorageConnectionPlatformSocket()
+{
+	Close();
+}
+
+bool FStorageConnectionPlatformSocket::Send(const uint8* Data, const uint64 DataSize)
+{
+	if (!Socket || Socket->GetState() != IPlatformHostSocket::EConnectionState::Connected)
+	{
+		return false;
+	}
+
+	bool bSuccess = Socket->Send(Data, DataSize) == IPlatformHostSocket::EResultNet::Ok;
+	return bSuccess;
+}
+
+bool FStorageConnectionPlatformSocket::Recv(uint8* Data, const uint64 DataSize, uint64& BytesRead, ESocketReceiveFlags::Type ReceiveFlags)
+{
+	if (!Socket || Socket->GetState() != IPlatformHostSocket::EConnectionState::Connected)
+	{
+		return false;
+	}
+
+	if (ConnectionBuffer.GetCapacity() < DataSize)
+	{
+		UE_LOG(LogStorageServerPlatformBackend, Display, TEXT("ConnectionBuffer capacity is lower than requested data read (%d vs %d)"), ConnectionBuffer.GetCapacity(), DataSize);
+	}
+
+	BytesRead = 0;
+	if (ConnectionBuffer.IsEmpty())
+	{
+		uint8 Buffer[1024];
+		if (Socket->Receive(Buffer, sizeof(Buffer), BytesRead, IPlatformHostSocket::EReceiveFlags::DontWait) != IPlatformHostSocket::EResultNet::Ok)
+		{
+			return false;
+		}
+
+		if (!ConnectionBuffer.Put(Buffer, BytesRead))
+		{
+			UE_LOG(LogStorageServerPlatformBackend, Display, TEXT("Couldn't fit the received data in the connection buffer"));
+		}
+	}
+
+	if (ReceiveFlags == ESocketReceiveFlags::Peek)
+	{
+		ConnectionBuffer.Peek(Data, DataSize, BytesRead);
+	}
+	else if (ReceiveFlags == ESocketReceiveFlags::WaitAll)
+	{
+		ConnectionBuffer.Consume(Data, DataSize, BytesRead);
+		uint64 TotalBytesRead = BytesRead;
+
+		while (TotalBytesRead < DataSize)
+		{
+			const uint64 BytesToRead = DataSize - TotalBytesRead;
+			if (Socket->Receive(Data + TotalBytesRead, BytesToRead, BytesRead, IPlatformHostSocket::EReceiveFlags::WaitAll) != IPlatformHostSocket::EResultNet::Ok)
+			{
+				return false;
+			}
+
+			TotalBytesRead += BytesRead;
+
+			if (TotalBytesRead > DataSize)
+			{
+				UE_LOG(LogStorageServerPlatformBackend, Display, TEXT("Exceeded what was supposed to be downloaded"));
+			}
+		}
+
+		BytesRead = TotalBytesRead;
+	}
+	else if (ReceiveFlags == ESocketReceiveFlags::None)
+	{
+		ConnectionBuffer.Consume(Data, DataSize, BytesRead);
+	}
+
+	return true;
+}
+
+bool FStorageConnectionPlatformSocket::HasPendingData(uint64& PendingDataSize) const
+{
+	PendingDataSize = ConnectionBuffer.GetSize();
+	return PendingDataSize > 0;
+}
+
+void FStorageConnectionPlatformSocket::Close()
+{
+	Communication->CloseConnection(Socket);
+}
+
+FStorageServerPlatformConnectionBackend::FStorageServerPlatformConnectionBackend(FStorageServerConnection& InOwner)
+	: FStorageConnectionBackend(InOwner)
+{
+}
+
+FStorageServerPlatformConnectionBackend::~FStorageServerPlatformConnectionBackend()
+{
+	for (auto& Socket : SocketPool)
+	{
+		delete Socket;
+	}
+}
+
+IStorageConnectionSocket* FStorageServerPlatformConnectionBackend::AcquireNewSocket(float TimeoutSeconds)
+{
+	FScopeLock Lock(&SocketPoolCS);
+
+	int32 FirstAvailableProtocolNumber = UsedSockets.Find(false);
+	if (FirstAvailableProtocolNumber == -1)
+	{
+		SocketPoolFreeConditionVariable.Wait(SocketPoolCS);
+		return AcquireSocketFromPool();
+	}
+
+	IPlatformHostSocketPtr PlatformSocket = Communication->OpenConnection(FirstAvailableProtocolNumber, *FString::Printf(TEXT("PlatformSocket %d"), FirstAvailableProtocolNumber));
+	if (!PlatformSocket)
+	{
+		return nullptr;
+	}
+
+	float WaitingFor = 0.f;
+	const float SleepTime = 0.01f;
+	IPlatformHostSocket::EConnectionState ConnectionState = PlatformSocket->GetState();
+	while (ConnectionState == IPlatformHostSocket::EConnectionState::Created)
+	{
+		if ((TimeoutSeconds != -1.f) && (WaitingFor > TimeoutSeconds))
+		{
+			Communication->CloseConnection(PlatformSocket);
+			UE_LOG(LogStorageServerPlatformBackend, Error, TEXT("Platform connection timed out"));
+			return nullptr;
+		}
+
+		FPlatformProcess::Sleep(SleepTime);
+		WaitingFor += SleepTime;
+		ConnectionState = PlatformSocket->GetState();
+	}
+
+	if (ConnectionState != IPlatformHostSocket::EConnectionState::Connected)
+	{
+		Communication->CloseConnection(PlatformSocket);
+		return nullptr;
+	}
+
+	FStorageConnectionPlatformSocket* NewSocket = new FStorageConnectionPlatformSocket(Communication, PlatformSocket, FirstAvailableProtocolNumber);
+	UsedSockets[FirstAvailableProtocolNumber] = true;
+
+	return NewSocket;
+}
+
+IStorageConnectionSocket* FStorageServerPlatformConnectionBackend::AcquireSocketFromPool()
+{
+	FScopeLock Lock(&SocketPoolCS);
+
+	if (!SocketPool.IsEmpty())
+	{
+		return SocketPool.Pop(EAllowShrinking::No);
+	}
+
+	return nullptr;
+}
+
+void FStorageServerPlatformConnectionBackend::ReleaseSocket(IStorageConnectionSocket* InSocket, bool bKeepAlive)
+{
+	FScopeLock Lock(&SocketPoolCS);
+
+	if (bKeepAlive)
+	{
+		uint64 PendingDataSize;
+		if (!InSocket->HasPendingData(PendingDataSize))
+		{
+			SocketPool.Push(InSocket);
+			SocketPoolFreeConditionVariable.NotifyOne();
+			return;
+		}
+		UE_LOG(LogStorageServerPlatformBackend, Error, TEXT("Socket was not fully drained"));
+	}
+
+	int32 ProtocolNumber = static_cast<FStorageConnectionPlatformSocket*>(InSocket)->GetProtocolNumber();
+	UsedSockets[ProtocolNumber] = false;
+	SocketPoolFreeConditionVariable.NotifyOne();
+	InSocket->Close();
+	delete InSocket;
+}
+
+bool FStorageServerPlatformConnectionBackend::InitializeInternal(TArrayView<const FString> InHostAddresses, int32 Port)
+{
+	Communication = &FPlatformMisc::GetPlatformHostCommunication();
+	if (!Communication)
+	{
+		return false;
+	}
+
+	UsedSockets.Init(false, 10);
+
+	int32 ServerVersion = HandshakeRequest(InHostAddresses);
+	if (ServerVersion != 1)
+	{
+		return false;
+	}
+
+	return true;
+}
+
+int32 FStorageServerPlatformConnectionBackend::HandshakeRequest(TArrayView<const FString> HostAddresses)
+{
+	TAnsiStringBuilder<256> ResourceBuilder;
+	ResourceBuilder.Append(OplogPath);
+
+	for (const FString& Addr : HostAddresses)
+	{
+		Hostname.Reset();
+		Hostname.Append(*Addr);
+
+		UE_LOG(LogStorageServerPlatformBackend, Display, TEXT("Trying to handshake with Zen at '%s'"), *Addr);
+		const float ConnectionTimeoutSeconds = 5.0f;
+
+		FStorageServerRequest Request("GET", *ResourceBuilder, Hostname);
+		if (IStorageConnectionSocket* Socket = Request.Send(Owner, false))
+		{
+			FStorageServerResponse Response(Owner, *Socket);
+
+			if (Response.IsOk())
+			{
+				FCbObject ResponseObj = Response.GetResponseObject();
+				return 1;
+			}
+			else
+			{
+				UE_LOG(LogStorageServerPlatformBackend, Fatal, TEXT("Failed to handshake with Zen at %s. '%s'"), *Addr, *Response.GetErrorMessage());
+			}
+		}
+	}
+
+	UE_LOG(LogStorageServerPlatformBackend, Fatal, TEXT("Failed to handshake with Zen at any of host addresses."));
+
+	return -1;
+}
diff --git a/Engine/Source/Runtime/StorageServerClient/Private/StoragePlatformConnectionBackend.h b/Engine/Source/Runtime/StorageServerClient/Private/StoragePlatformConnectionBackend.h
new file mode 100644
--- /dev/null
+++ b/Engine/Source/Runtime/StorageServerClient/Private/StoragePlatformConnectionBackend.h
@@ -0,0 +1,173 @@
+// Copyright Epic Games, Inc. All Rights Reserved.
+
+#pragma once
+
+#include "StorageServerConnection.h"
+#include "Experimental/Async/ConditionVariable.h"
+#include "GenericPlatform/GenericPlatformHostCommunication.h"
+#include "GenericPlatform/GenericPlatformHostSocket.h"
+
+class FConnectionCircularBuffer
+{
+public:
+	explicit FConnectionCircularBuffer(uint64 InCapacity)
+		: CapacityMask(FMath::RoundUpToPowerOfTwo(InCapacity) - 1)
+	{
+		Allocation = new uint8[GetCapacity()];
+	}
+
+	~FConnectionCircularBuffer()
+	{
+		delete[] Allocation;
+	}
+
+	uint8 operator[](uint64 Index) const
+	{
+		return const_cast<FConnectionCircularBuffer*>(this)->operator[](Index);
+	}
+
+	uint8& operator[](uint64 Index)
+	{
+		check(Index < GetSize());
+		uint64 DataOffset = (Index + Tail) & (CapacityMask);
+		return Allocation[DataOffset];
+	}
+
+	uint64 GetSize() const
+	{
+		return Size;
+	}
+
+	uint64 GetCapacity() const
+	{
+		return CapacityMask + 1;
+	}
+
+	bool IsEmpty() const
+	{
+		return Size == 0;
+	}
+
+	void Clear()
+	{
+		Tail = 0;
+		Head = 0;
+		Size = 0;
+	}
+
+	uint64 SpaceLeft() const
+	{
+		return GetCapacity() - Size;
+	}
+
+	void Peek(uint8* Data, const uint64 DataSize, uint64& OutSize)
+	{
+		OutSize = FMath::Min(DataSize, Size);
+
+		if (OutSize + Tail > GetCapacity())
+		{
+			uint64 LenToEnd = GetCapacity() - Tail;
+			uint64 LenFromBegin = OutSize - LenToEnd;
+			FMemory::Memcpy(Data, Allocation + Tail, LenToEnd);
+			FMemory::Memcpy(Data + LenToEnd, Allocation, LenFromBegin);
+		}
+		else
+		{
+			FMemory::Memcpy(Data, Allocation + Tail, OutSize);
+		}
+	}
+
+	void Consume(uint8* Data, const uint64 DataSize, uint64& OutSize)
+	{
+		Peek(Data, DataSize, OutSize);
+		Tail = (Tail + OutSize) & CapacityMask;
+
+		Size -= OutSize;
+
+		if (Size == 0)
+		{
+			Head = 0;
+			Tail = 0;
+		}
+	}
+
+	bool Put(uint8* Data, uint64 DataSize)
+	{
+		if (DataSize > SpaceLeft())
+		{
+			return false;
+		}
+
+		if (DataSize + Head > GetCapacity())
+		{
+			uint64 LenToEnd = GetCapacity() - Head;
+			uint64 LenFromBegin = DataSize - LenToEnd;
+			FMemory::Memcpy(Allocation + Head, Data, LenToEnd);
+			FMemory::Memcpy(Allocation, Data + LenToEnd, LenFromBegin);
+		}
+		else
+		{
+			FMemory::Memcpy(Allocation + Head, Data, DataSize);
+		}
+
+		Head = (Head + DataSize) & (CapacityMask);
+
+		Size += DataSize;
+
+		return true;
+	}
+
+private:
+	uint8* Allocation{ nullptr };
+	uint64 CapacityMask{ 0 };
+	uint64 Size{ 0 };
+	uint64 Head{ 0 };
+	uint64 Tail{ 0 };
+};
+
+class FStorageConnectionPlatformSocket : public IStorageConnectionSocket
+{
+public:
+	FStorageConnectionPlatformSocket(IPlatformHostCommunication* HostCommunication, IPlatformHostSocketPtr InSocket, int32 InProtocolNumber);
+	~FStorageConnectionPlatformSocket() override;
+
+	bool Send(const uint8* Data, const uint64 DataSize) override;
+	bool Recv(uint8* Data, const uint64 DataSize, uint64& BytesRead, ESocketReceiveFlags::Type ReceiveFlags) override;
+	bool HasPendingData(uint64& PendingDataSize) const override;
+	void Close() override;
+
+	int32 GetProtocolNumber() const { return ProtocolNumber; }
+
+private:
+	IPlatformHostCommunication* Communication;
+	IPlatformHostSocketPtr Socket;
+	FConnectionCircularBuffer ConnectionBuffer;
+	int32 ProtocolNumber;
+};
+
+class FStorageServerPlatformConnectionBackend : public FStorageConnectionBackend
+{
+public:
+	FStorageServerPlatformConnectionBackend(FStorageServerConnection& InOwner);
+	~FStorageServerPlatformConnectionBackend() override;
+
+	IStorageConnectionSocket* AcquireNewSocket(float TimeoutSeconds) override;
+	IStorageConnectionSocket* AcquireSocketFromPool() override;
+	void ReleaseSocket(IStorageConnectionSocket* Socket, bool bKeepAlive) override;
+
+protected:
+	bool InitializeInternal(TArrayView<const FString> InHostAddresses, int32 Port) override;
+
+private:
+	int32 HandshakeRequest(TArrayView<const FString> HostAddresses);
+
+private:
+	IPlatformHostCommunication* Communication{ nullptr };
+
+	FCriticalSection SocketPoolCS;
+	UE::FConditionVariable SocketPoolFreeConditionVariable;
+	TArray<IStorageConnectionSocket*> SocketPool;
+	TBitArray<> UsedSockets; // bitset to keep track of used sockets in the pool
+
+	TAnsiStringBuilder<1024> Hostname;
+};
\ No newline at end of file
diff --git a/Engine/Source/Runtime/StorageServerClient/Private/StorageServerConnection.cpp b/Engine/Source/Runtime/StorageServerClient/Private/StorageServerConnection.cpp
--- a/Engine/Source/Runtime/StorageServerClient/Private/StorageServerConnection.cpp
+++ b/Engine/Source/Runtime/StorageServerClient/Private/StorageServerConnection.cpp
@@ -16,6 +16,9 @@
 #include "SocketSubsystem.h"
 #include "Sockets.h"
 #include "IO/PackageStore.h"
+#include "StorageSocketConnectionBackend.h"
+#include "StoragePlatformConnectionBackend.h"
+#include "GenericPlatform/GenericPlatformHostCommunication.h"
 
 #if !UE_BUILD_SHIPPING
 
@@ -23,41 +26,6 @@ DEFINE_LOG_CATEGORY_STATIC(LogStorageServerConnection, Log, All);
 
 TRACE_DECLARE_MEMORY_COUNTER(ZenHttpClientSerializedBytes, TEXT("ZenClient/SerializedBytes"));
 
-static TArray<TSharedPtr<FInternetAddr>> GetAddressFromString(ISocketSubsystem& SocketSubsystem, TArrayView<const FString> HostAddresses, const int32 Port)
-{
-	TArray<TSharedPtr<FInternetAddr>> InternetAddresses;
-	FString ModifiedHostAddr;
-
-	for (const FString& HostAddr : HostAddresses)
-	{
-		// Numeric IPV6 addresses can be enclosed in brackets, and must have the brackets stripped before calling GetAddressFromString
-		const FString* EffectiveHostAddr = &HostAddr;
-		if (!HostAddr.IsEmpty() && HostAddr[0] == TEXT('[') && HostAddr[HostAddr.Len() - 1] == TEXT(']'))
-		{
-			ModifiedHostAddr = FStringView(HostAddr).Mid(1,HostAddr.Len() - 2);
-			EffectiveHostAddr = &ModifiedHostAddr;
-		}
-		TSharedPtr<FInternetAddr> Addr = SocketSubsystem.GetAddressFromString(*EffectiveHostAddr);
-
-		if (!Addr.IsValid() || !Addr->IsValid())
-		{
-			FAddressInfoResult GAIRequest = SocketSubsystem.GetAddressInfo(**EffectiveHostAddr, nullptr, EAddressInfoFlags::Default, NAME_None);
-			if (GAIRequest.ReturnCode == SE_NO_ERROR && GAIRequest.Results.Num() > 0)
-			{
-				Addr = GAIRequest.Results[0].Address;
-			}
-		}
-
-		if (Addr.IsValid() && Addr->IsValid())
-		{
-			Addr->SetPort(Port);
-			InternetAddresses.Emplace(MoveTemp(Addr));
-		}
-	}
-
-	return InternetAddresses;
-}
-
 static uint64 GetCompressedOffset(const FCompressedBuffer& Buffer, uint64 RawOffset)
 {
 	if (RawOffset > 0)
@@ -84,7 +52,7 @@ FStorageServerRequest::FStorageServerRequest(FAnsiStringView Verb, FAnsiStringVi
 		<< "Accept: " << GetMimeTypeString(Accept) << "\r\n";
 }
 
-FSocket* FStorageServerRequest::Send(FStorageServerConnection& Owner, bool bLogOnError)
+IStorageConnectionSocket* FStorageServerRequest::Send(FStorageServerConnection& Owner, bool bLogOnError)
 {
 	if (BodyBuffer.Num())
 	{
@@ -93,32 +61,30 @@ FSocket* FStorageServerRequest::Send(FStorageServerConnection& Owner, bool bLogO
 	HeaderBuffer << "\r\n";
 	int32 BytesLeft = HeaderBuffer.Len();
 	
-	auto Send = [](FSocket* Socket, const uint8* Data, int32 Length)
+	auto Send = [](IStorageConnectionSocket* Socket, const uint8* Data, int32 Length)
 	{
 		int32 BytesLeft = Length;
 		while (BytesLeft > 0)
 		{
-			int32 BytesSent;
-			if (!Socket->Send(Data, BytesLeft, BytesSent))
+			if (!Socket->Send(Data, BytesLeft))
 			{
 				return false;
 			}
-			check(BytesSent >= 0);
-			BytesLeft -= BytesSent;
-			Data += BytesSent;
+			BytesLeft -= BytesLeft;
 		}
 		return true;
 	};
 
 	int32 Attempts = 0;
+	FStorageConnectionBackend* Backend = Owner.GetConnectionBackend();
 	while (Attempts < 10)
 	{
-		FSocket* Socket;
-		FSocket* SocketFromPool;
-		Socket = SocketFromPool = Owner.AcquireSocketFromPool();
+		IStorageConnectionSocket* Socket;
+		IStorageConnectionSocket* SocketFromPool;
+		Socket = SocketFromPool = Backend->AcquireSocketFromPool();
 		if (!Socket)
 		{
-			Socket = Owner.AcquireNewSocket();
+			Socket = Backend->AcquireNewSocket();
 			if (!Socket)
 			{
 				++Attempts;
@@ -129,7 +95,7 @@ FSocket* FStorageServerRequest::Send(FStorageServerConnection& Owner, bool bLogO
 		if (!Send(Socket, reinterpret_cast<const uint8*>(HeaderBuffer.GetData()), HeaderBuffer.Len()) ||
 			!Send(Socket, BodyBuffer.GetData(), BodyBuffer.Num()))
 		{
-			Owner.ReleaseSocket(Socket, false);
+			Backend->ReleaseSocket(Socket, false);
 			if (!SocketFromPool)
 			{
 				++Attempts;
@@ -138,6 +104,7 @@ FSocket* FStorageServerRequest::Send(FStorageServerConnection& Owner, bool bLogO
 		}
 		return Socket;
 	}
+
 	if (bLogOnError)
 	{
 		UE_LOG(LogStorageServerConnection, Fatal, TEXT("Failed sending request to storage server."));
@@ -152,7 +119,7 @@ void FStorageServerRequest::Serialize(void* V, int64 Length)
 	FMemory::Memcpy(Dest, V, Length);
 }
 
-FStorageServerResponse::FStorageServerResponse(FStorageServerConnection& InOwner, FSocket& InSocket)
+FStorageServerResponse::FStorageServerResponse(FStorageServerConnection& InOwner, IStorageConnectionSocket& InSocket)
 	: Owner(InOwner)
 	, Socket(&InSocket)
 {
@@ -163,7 +130,7 @@ FStorageServerResponse::FStorageServerResponse(FStorageServerConnection& InOwner
 	{
 		for (;;)
 		{
-			int32 BytesRead;
+			uint64 BytesRead;
 			InSocket.Recv(Buffer, 1024, BytesRead, ESocketReceiveFlags::Peek);
 			FAnsiStringView ResponseView(reinterpret_cast<const ANSICHAR*>(Buffer), BytesRead);
 			int32 LineEndIndex;
@@ -203,7 +170,7 @@ FStorageServerResponse::FStorageServerResponse(FStorageServerConnection& InOwner
 	{
 		TArray<uint8> ErrorBuffer;
 		ErrorBuffer.SetNumUninitialized(ContentLength + 1);
-		int32 BytesRead = 0;
+		uint64 BytesRead = 0;
 		InSocket.Recv(ErrorBuffer.GetData(), ContentLength, BytesRead, ESocketReceiveFlags::WaitAll);
 
 		if (BytesRead > 0)
@@ -247,7 +214,7 @@ FStorageServerChunkBatchRequest& FStorageServerChunkBatchRequest::AddChunk(const
 
 bool FStorageServerChunkBatchRequest::Issue(TFunctionRef<void(uint32 ChunkCount, uint32* ChunkIndices, uint64* ChunkSizes, FStorageServerResponse& ChunkDataStream)> OnResponse)
 {
-	FSocket* Socket = Send(Owner);
+	IStorageConnectionSocket* Socket = Send(Owner);
 	if (!Socket)
 	{
 		UE_LOG(LogStorageServerConnection, Fatal, TEXT("Failed to send chunk batch request to storage server."));
@@ -300,7 +267,7 @@ bool FStorageServerChunkBatchRequest::Issue(TFunctionRef<void(uint32 ChunkCount,
 
 void FStorageServerResponse::ReleaseSocket(bool bKeepAlive)
 {
-	Owner.ReleaseSocket(Socket, bKeepAlive);
+	Owner.GetConnectionBackend()->ReleaseSocket(Socket, bKeepAlive);
 	Socket = nullptr;
 }
 
@@ -325,7 +292,7 @@ void FStorageServerResponse::Serialize(void* V, int64 Length)
 	while (RemainingBytesToRead)
 	{
 		uint64 BytesToRead32 = FMath::Min(RemainingBytesToRead, static_cast<uint64>(INT32_MAX));
-		int32 BytesRead;
+		uint64 BytesRead;
 		if (!Socket->Recv(Destination, static_cast<int32>(BytesToRead32), BytesRead, ESocketReceiveFlags::WaitAll))
 		{
 			UE_LOG(LogStorageServerConnection, Fatal, TEXT("Failed reading %d bytes from socket"), BytesToRead32);
@@ -351,12 +318,6 @@ int64 FStorageServerResponse::SerializeChunk(FStorageServerSerializationContext&
 		return 0;
 	}
 
-	if (!Socket)
-	{
-		UE_LOG(LogStorageServerConnection, Fatal, TEXT("Trying to read %lld bytes from released socket"), ContentLength);
-		return 0;
-	}
-
 	if (ContentType == EStorageServerContentType::Binary)
 	{
 		const uint64 ChunkSize = FMath::Min<uint64>(ContentLength, RawSize);
@@ -397,12 +358,6 @@ int64 FStorageServerResponse::SerializeChunkTo(FMutableMemoryView Memory, uint64
 		return 0;
 	}
 
-	if (!Socket)
-	{
-		UE_LOG(LogStorageServerConnection, Fatal, TEXT("Trying to read %lld bytes from released socket"), ContentLength);
-		return 0;
-	}
-
 	if (ContentType == EStorageServerContentType::Binary)
 	{
 		FMutableMemoryView Dst = Memory.Left(FMath::Min<uint64>(Memory.GetSize(), ContentLength));
@@ -470,30 +425,20 @@ FCbObject FStorageServerResponse::GetResponseObject()
 	return Payload.AsObject();
 }
 
-FStorageServerConnection::FStorageServerConnection()
-	: SocketSubsystem(*ISocketSubsystem::Get())
+FStorageConnectionBackend::FStorageConnectionBackend(FStorageServerConnection& InOwner)
+	: Owner(InOwner)
 {
 }
 
-FStorageServerConnection::~FStorageServerConnection()
+bool FStorageConnectionBackend::Initialize(TArrayView<const FString> InHostAddresses, int32 InPort, const TCHAR* InProjectNameOverride, const TCHAR* InPlatformNameOverride)
 {
-	for (FSocket* Socket : SocketPool)
-	{
-		Socket->Close();
-		delete Socket;
-	}
+	InitOplog(InProjectNameOverride, InPlatformNameOverride);
+
+	return InitializeInternal(InHostAddresses, InPort);
 }
 
-bool FStorageServerConnection::Initialize(TArrayView<const FString> InHostAddresses, int32 InPort, const TCHAR* InProjectNameOverride, const TCHAR* InPlatformNameOverride)
+void FStorageConnectionBackend::InitOplog(const TCHAR* InProjectNameOverride, const TCHAR* InPlatformNameOverride)
 {
-	TArray<TSharedPtr<FInternetAddr>> HostAddresses = GetAddressFromString(SocketSubsystem, InHostAddresses, InPort);
-
-	if (!HostAddresses.Num())
-	{
-		UE_LOG(LogStorageServerConnection, Fatal, TEXT("No valid Zen store host address specified"));
-		return false;
-	}
-	
 	OplogPath.Append("/prj/");
 	if (InProjectNameOverride)
 	{
@@ -515,161 +460,47 @@ bool FStorageServerConnection::Initialize(TArrayView<const FString> InHostAddres
 		check(TargetPlatformNames.Num() > 0);
 		OplogPath.Append(TCHAR_TO_ANSI(*TargetPlatformNames[0]));
 	}
-
-	const int32 ServerVersion = HandshakeRequest(HostAddresses);
-	if (ServerVersion != 1)
-	{
-		return false;
-	}
-
-	UE_LOG(LogStorageServerConnection, Display, TEXT("Connected to Zen storage server at '%s'"), *ServerAddr->ToString(true));
-
-	return true;
 }
 
-void FStorageServerConnection::SortHostAddressesByLocalSubnet(TArrayView<const TSharedPtr<FInternetAddr>> HostAddresses, TArray<TSharedPtr<FInternetAddr>>& SortedHostAddresses)
+FStorageServerConnection::FStorageServerConnection()
+	: SocketSubsystem(*ISocketSubsystem::Get())
 {
-	//no sorting needed cases
-	if (HostAddresses.Num() == 0)
-		return;
-
-	if (HostAddresses.Num() == 1)
-	{
-		SortedHostAddresses.Push(HostAddresses[0]);
-		return;
-	}
-
-	//Sorting logic:
-	//1 only on desktop, if it's an IPV6 address loopback (ends with ":1")
-	//2 only on desktop, if it's and IPV4 address loopback (starts with "127.0.0")
-	//3 if the host IPV4 subnet match the client subnet (xxx.xxx.xxx)
-	//4 remaining addresses
-	bool bCanBindAll = false;
-	bool bAppendPort = false;
-	TSharedPtr<FInternetAddr> localAddr = SocketSubsystem.GetLocalHostAddr(*GLog, bCanBindAll);
-	FString localAddrStringSubnet = localAddr->ToString(bAppendPort);
-
-	int32 localLastDotPos = INDEX_NONE;
-	if (localAddrStringSubnet.FindLastChar(TEXT('.'), localLastDotPos))
-	{
-		localAddrStringSubnet = localAddrStringSubnet.LeftChop(localAddrStringSubnet.Len() - localLastDotPos);
-	}
+}
 
-	TArray<TSharedPtr<FInternetAddr>> IPV6Loopback;
-	TArray<TSharedPtr<FInternetAddr>> IPV4Loopback;
-	TArray<TSharedPtr<FInternetAddr>> RegularAddresses;
+FStorageServerConnection::~FStorageServerConnection()
+{
+}
 
-	for (const TSharedPtr<FInternetAddr>& Addr : HostAddresses)
+bool FStorageServerConnection::Initialize(TArrayView<const FString> InHostAddresses, int32 InPort, const TCHAR* InProjectNameOverride, const TCHAR* InPlatformNameOverride)
+{
+	OplogPath.Append("/prj/");
+	if (InProjectNameOverride)
 	{
-		FString tempAddrStringSubnet = Addr->ToString(bAppendPort);
-
-#if PLATFORM_DESKTOP
-		if (Addr->GetProtocolType() == FNetworkProtocolTypes::IPv6)
-		{
-			if (tempAddrStringSubnet.EndsWith(":1"))
-			{
-				IPV6Loopback.Push(Addr);
-				continue;
-			}
-		}
-		else
-		{
-			if (tempAddrStringSubnet.StartsWith("127.0.0."))
-			{
-				IPV4Loopback.Push(Addr);
-				continue;
-			}
-		}
-#endif
-		int32 LastDotPos = INDEX_NONE;
-		if (tempAddrStringSubnet.FindLastChar(TEXT('.'), LastDotPos))
-		{
-			tempAddrStringSubnet = tempAddrStringSubnet.LeftChop(tempAddrStringSubnet.Len() - LastDotPos);
-		}
-
-		if (localAddrStringSubnet.Equals(tempAddrStringSubnet))
-			RegularAddresses.Insert(Addr, 0);
-		else
-			RegularAddresses.Push(Addr);
+		OplogPath.Append(TCHAR_TO_ANSI(InProjectNameOverride));
 	}
-
-
-	for (const TSharedPtr<FInternetAddr>& Addrv6lb : IPV6Loopback)
+	else
 	{
-		SortedHostAddresses.Push(Addrv6lb);
+		OplogPath.Append(TCHAR_TO_ANSI(*FApp::GetZenStoreProjectId()));
 	}
-
-	for (const TSharedPtr<FInternetAddr>& Addrv4lb : IPV4Loopback)
+	OplogPath.Append("/oplog/");
+	if (InPlatformNameOverride)
 	{
-		SortedHostAddresses.Push(Addrv4lb);
+		OplogPath.Append(TCHAR_TO_ANSI(InPlatformNameOverride));
 	}
-
-	for (const TSharedPtr<FInternetAddr>& RegularAddr : RegularAddresses)
+	else
 	{
-		SortedHostAddresses.Push(RegularAddr);
+		TArray<FString> TargetPlatformNames;
+		FPlatformMisc::GetValidTargetPlatforms(TargetPlatformNames);
+		check(TargetPlatformNames.Num() > 0);
+		OplogPath.Append(TCHAR_TO_ANSI(*TargetPlatformNames[0]));
 	}
-}
-
-int32 FStorageServerConnection::HandshakeRequest(TArrayView<const TSharedPtr<FInternetAddr>> HostAddresses)
-{
-	TAnsiStringBuilder<256> ResourceBuilder;
-	ResourceBuilder.Append(OplogPath);
 
-	// We sort the host address trying to match the local subnet in order to improve the chances of a successful connection at the first attempt.
-	TArray<TSharedPtr<FInternetAddr>> SortedAddresses;
-	SortHostAddressesByLocalSubnet(HostAddresses, SortedAddresses);
-
-	for (TSharedPtr<FInternetAddr>& Addr : SortedAddresses)
+	if (!CreateConnectionBackend(InHostAddresses, InPort))
 	{
-		Hostname.Reset();
-		Hostname.Append(TCHAR_TO_ANSI(*Addr->ToString(false)));
-		ServerAddr = Addr;
-		
-		UE_LOG(LogStorageServerConnection, Display, TEXT("Trying to handshake with Zen at '%s'"), *Addr->ToString(true));
-
-		// Handshakes are done with a limited connection timeout so that we can find out if the destination is unreachable
-		// in a timely manner.
-		const float ConnectionTimeoutSeconds = 5.0f;
-		FSocket* ConnectSocket = AcquireNewSocket(ConnectionTimeoutSeconds);
-		if (!ConnectSocket)
-		{
-			continue;
-		}
-		ReleaseSocket(ConnectSocket, true);
-
-		FStorageServerRequest Request("GET", *ResourceBuilder, Hostname);
-		if (FSocket* Socket = Request.Send(*this, false))
-		{
-			FStorageServerResponse Response(*this, *Socket);
-
-			if (Response.IsOk())
-			{
-				FCbObject ResponseObj = Response.GetResponseObject();
-
-				// we currently don't have any concept of protocol versioning, if
-				// we succeed in communicating with the endpoint we're good since
-				// any breaking API change would need to be done in a backward
-				// compatible manner
-
-				return 1;
-			}
-			else
-			{
-				UE_LOG(LogStorageServerConnection, Fatal, TEXT("Failed to handshake with Zen at %s. '%s'"), *ServerAddr->ToString(true), *Response.GetErrorMessage());
-			}
-		}
-		else
-		{
-			UE_LOG(LogStorageServerConnection, Warning, TEXT("Failed to send handshake request to Zen at %s."), *ServerAddr->ToString(true));
-		}
+		return false;
 	}
 
-	UE_LOG(LogStorageServerConnection, Error, TEXT("Failed to handshake with Zen at any of host addresses."));
-
-	Hostname.Reset();
-	ServerAddr.Reset();
-
-	return -1;
+	return ConnectionBackend->Initialize(InHostAddresses, InPort, InProjectNameOverride, InPlatformNameOverride);
 }
 
 void FStorageServerConnection::PackageStoreRequest(TFunctionRef<void(FPackageStoreEntryResource&&)> Callback)
@@ -679,7 +510,7 @@ void FStorageServerConnection::PackageStoreRequest(TFunctionRef<void(FPackageSto
 	TAnsiStringBuilder<256> ResourceBuilder;
 	ResourceBuilder.Append(OplogPath).Append("/entries?fieldfilter=packagestoreentry");
 	FStorageServerRequest Request("GET", *ResourceBuilder, Hostname, EStorageServerContentType::CompressedBinary);
-	FSocket* Socket = Request.Send(*this);
+	IStorageConnectionSocket* Socket = Request.Send(*this);
 	if (!Socket)
 	{
 		UE_LOG(LogStorageServerConnection, Fatal, TEXT("Failed to send oplog request to storage server at %s."), *ServerAddr->ToString(true));
@@ -709,7 +540,7 @@ void FStorageServerConnection::FileManifestRequest(TFunctionRef<void(FIoChunkId
 	TAnsiStringBuilder<256> ResourceBuilder;
 	ResourceBuilder.Append(OplogPath).Append("/files?filter=client");
 	FStorageServerRequest Request("GET", *ResourceBuilder, Hostname, EStorageServerContentType::CbObject);
-	FSocket* Socket = Request.Send(*this);
+	IStorageConnectionSocket* Socket = Request.Send(*this);
 	if (!Socket)
 	{
 		UE_LOG(LogStorageServerConnection, Fatal, TEXT("Failed to send file manifest request to storage server at %s."), *ServerAddr->ToString(true));
@@ -748,7 +579,7 @@ int64 FStorageServerConnection::ChunkSizeRequest(const FIoChunkId& ChunkId)
 	ResourceBuilder << "/" << ChunkId << "/info";
 
 	FStorageServerRequest Request("GET", *ResourceBuilder, Hostname, EStorageServerContentType::CbObject);
-	FSocket* Socket = Request.Send(*this);
+	IStorageConnectionSocket* Socket = Request.Send(*this);
 	if (!Socket)
 	{
 		UE_LOG(LogStorageServerConnection, Fatal, TEXT("Failed to send chunk size request to storage server at %s."), *ServerAddr->ToString(true));
@@ -809,7 +640,7 @@ bool FStorageServerConnection::ReadChunkRequest(const FIoChunkId& ChunkId, uint6
 	}
 
 	FStorageServerRequest Request("GET", *ResourceBuilder, Hostname, EStorageServerContentType::CompressedBinary);
-	FSocket* Socket = Request.Send(*this);
+	IStorageConnectionSocket* Socket = Request.Send(*this);
 	if (!Socket)
 	{
 		UE_LOG(LogStorageServerConnection, Fatal, TEXT("Failed to send chunk read request to storage server at %s."), *ServerAddr->ToString(true));
@@ -839,66 +670,48 @@ FStorageServerChunkBatchRequest FStorageServerConnection::NewChunkBatchRequest()
 	return FStorageServerChunkBatchRequest(*this, *ResourceBuilder, Hostname);
 }
 
-FSocket* FStorageServerConnection::AcquireSocketFromPool()
+FString FStorageServerConnection::GetHostAddr() const
 {
-	FScopeLock Lock(&SocketPoolCritical);
-	if (!SocketPool.IsEmpty())
-	{
-		return SocketPool.Pop(EAllowShrinking::No);
-	}
-	return nullptr;
+	return ServerAddr.IsValid() ? ServerAddr->ToString(false) : FString();
 }
 
-FSocket* FStorageServerConnection::AcquireNewSocket(float TimeoutSeconds)
+bool FStorageServerConnection::CreateConnectionBackend(TArrayView<const FString> HostAddresses, int32 Port)
 {
-	FSocket* Socket = SocketSubsystem.CreateSocket(NAME_Stream, TEXT("StorageServer"), ServerAddr->GetProtocolType());
-	check(Socket);
-
-	if (TimeoutSeconds > 0.0f)
+	FString PlatformAddress;
+	for (const FString& Addr : HostAddresses)
 	{
-		Socket->SetNonBlocking(true);
-		ON_SCOPE_EXIT
-		{
-			Socket->SetNonBlocking(false);
-		};
-
-		if (Socket->Connect(*ServerAddr) && Socket->Wait(ESocketWaitConditions::WaitForWrite, FTimespan::FromSeconds(TimeoutSeconds)))
+		if (Addr.Contains(TEXT("platform://")))
 		{
-			return Socket;
+			PlatformAddress = Addr;
 		}
 	}
+
+	if (!PlatformAddress.IsEmpty())
+	{
+		ConnectionBackend = MakeUnique<FStorageServerPlatformConnectionBackend>(*this);
+	}
 	else
 	{
-		if (Socket->Connect(*ServerAddr))
-		{
-			return Socket;
-		}
+		ConnectionBackend = MakeUnique<FStorageSocketConnectionBackend>(*this, *ISocketSubsystem::Get());
 	}
 
-	delete Socket;
-	return nullptr;
+	return ConnectionBackend != nullptr;
 }
 
-FString FStorageServerConnection::GetHostAddr() const
+bool FStorageServerConnection::CreatePlatformBackend(const FString& HostAddresses, int32 Port)
 {
-	return ServerAddr.IsValid() ? ServerAddr->ToString(false) : FString();
+	ConnectionBackend = MakeUnique<FStorageServerPlatformConnectionBackend>(*this);
+	if (!ConnectionBackend)
+	{
+		return false;
+	}
+
+	return true;
 }
 
-void FStorageServerConnection::ReleaseSocket(FSocket* Socket, bool bKeepAlive)
+FStorageConnectionBackend* FStorageServerConnection::GetConnectionBackend() const
 {
-	if (bKeepAlive)
-	{
-		uint32 PendingDataSize;
-		if (!Socket->HasPendingData(PendingDataSize))
-		{
-			FScopeLock Lock(&SocketPoolCritical);
-			SocketPool.Push(Socket);
-			return;
-		}
-		UE_LOG(LogStorageServerConnection, Fatal, TEXT("Socket was not fully drained"));
-	}
-	Socket->Close();
-	delete Socket;
+	return ConnectionBackend.Get();
 }
 
 #endif
diff --git a/Engine/Source/Runtime/StorageServerClient/Private/StorageSocketConnectionBackend.cpp b/Engine/Source/Runtime/StorageServerClient/Private/StorageSocketConnectionBackend.cpp
new file mode 100644
--- /dev/null
+++ b/Engine/Source/Runtime/StorageServerClient/Private/StorageSocketConnectionBackend.cpp
@@ -0,0 +1,333 @@
+// Copyright Epic Games, Inc. All Rights Reserved.
+
+#include "StorageSocketConnectionBackend.h"
+
+#include "Misc/OutputDeviceRedirector.h"
+#include "Misc/ScopeExit.h"
+#include "Serialization/CompactBinary.h"
+#include "Serialization/CompactBinarySerialization.h"
+#include "SocketSubsystem.h"
+#include "Sockets.h"
+
+DEFINE_LOG_CATEGORY_STATIC(LogStorageSocketBackend, Log, All);
+
+static TArray<TSharedPtr<FInternetAddr>> GetAddressFromString(ISocketSubsystem& SocketSubsystem, TArrayView<const FString> HostAddresses, const int32 Port)
+{
+	TArray<TSharedPtr<FInternetAddr>> InternetAddresses;
+	FString ModifiedHostAddr;
+
+	for (const FString& HostAddr : HostAddresses)
+	{
+		// Numeric IPV6 addresses can be enclosed in brackets, and must have the brackets stripped before calling GetAddressFromString
+		const FString* EffectiveHostAddr = &HostAddr;
+		if (!HostAddr.IsEmpty() && HostAddr[0] == TEXT('[') && HostAddr[HostAddr.Len() - 1] == TEXT(']'))
+		{
+			ModifiedHostAddr = FStringView(HostAddr).Mid(1, HostAddr.Len() - 2);
+			EffectiveHostAddr = &ModifiedHostAddr;
+		}
+		TSharedPtr<FInternetAddr> Addr = SocketSubsystem.GetAddressFromString(*EffectiveHostAddr);
+
+		if (!Addr.IsValid() || !Addr->IsValid())
+		{
+			FAddressInfoResult GAIRequest = SocketSubsystem.GetAddressInfo(**EffectiveHostAddr, nullptr, EAddressInfoFlags::Default, NAME_None);
+			if (GAIRequest.ReturnCode == SE_NO_ERROR && GAIRequest.Results.Num() > 0)
+			{
+				Addr = GAIRequest.Results[0].Address;
+			}
+		}
+
+		if (Addr.IsValid() && Addr->IsValid())
+		{
+			Addr->SetPort(Port);
+			InternetAddresses.Emplace(MoveTemp(Addr));
+		}
+	}
+
+	return InternetAddresses;
+}
+
+FStorageConnectionSocketFSocket::FStorageConnectionSocketFSocket(FSocket* InSocket)
+	: Socket(InSocket)
+{
+}
+
+FStorageConnectionSocketFSocket::~FStorageConnectionSocketFSocket()
+{
+	if (Socket)
+	{
+		Socket->Close();
+		delete Socket;
+		Socket = nullptr;
+	}
+}
+
+bool FStorageConnectionSocketFSocket::Send(const uint8* Data, const uint64 DataSize)
+{
+	if (!Socket)
+	{
+		return false;
+	}
+
+	int32 BytesSent;
+	return Socket->Send(Data, DataSize, BytesSent);
+}
+
+bool FStorageConnectionSocketFSocket::Recv(uint8* Data, const uint64 DataSize, uint64& BytesRead, ESocketReceiveFlags::Type ReceiveFlags)
+{
+	if (!Socket)
+	{
+		return false;
+	}
+
+	int32 ReadBytes = 0;
+	if (!Socket->Recv(Data, DataSize, ReadBytes, ReceiveFlags))
+	{
+		return false;
+	}
+
+	BytesRead = ReadBytes;
+	return true;
+}
+
+bool FStorageConnectionSocketFSocket::HasPendingData(uint64& PendingDataSize) const
+{
+	uint32 PendingData;
+	bool bRes = Socket->HasPendingData(PendingData);
+
+	PendingDataSize = PendingData;
+	return bRes;
+}
+
+void FStorageConnectionSocketFSocket::Close()
+{
+	Socket->Close();
+}
+
+FStorageSocketConnectionBackend::~FStorageSocketConnectionBackend()
+{
+	for (IStorageConnectionSocket* Socket : SocketPool)
+	{
+		delete Socket;
+	}
+}
+
+IStorageConnectionSocket* FStorageSocketConnectionBackend::AcquireSocketFromPool()
+{
+	FScopeLock Lock(&SocketPoolCS);
+
+	if (!SocketPool.IsEmpty())
+	{
+		return SocketPool.Pop(EAllowShrinking::No);
+	}
+
+	return nullptr;
+}
+
+IStorageConnectionSocket* FStorageSocketConnectionBackend::AcquireNewSocket(float TimeoutSeconds)
+{
+	FSocket* Socket = SocketSubsystem.CreateSocket(NAME_Stream, TEXT("StorageServer"), ServerAddr->GetProtocolType());
+
+	if (TimeoutSeconds > 0.0f)
+	{
+		Socket->SetNonBlocking(true);
+		ON_SCOPE_EXIT
+		{
+			Socket->SetNonBlocking(false);
+		};
+
+		if (Socket->Connect(*ServerAddr) && Socket->Wait(ESocketWaitConditions::WaitForWrite, FTimespan::FromSeconds(TimeoutSeconds)))
+		{
+			return new FStorageConnectionSocketFSocket(Socket);
+		}
+	}
+	else
+	{
+		if (Socket->Connect(*ServerAddr))
+		{
+			return new FStorageConnectionSocketFSocket(Socket);
+		}
+	}
+
+	delete Socket;
+	return nullptr;
+}
+
+void FStorageSocketConnectionBackend::ReleaseSocket(IStorageConnectionSocket* Socket, bool bKeepAlive)
+{
+	if (bKeepAlive)
+	{
+		uint64 PendingDataSize;
+		if (!Socket->HasPendingData(PendingDataSize))
+		{
+			FScopeLock Lock(&SocketPoolCS);
+			SocketPool.Push(Socket);
+			return;
+		}
+		UE_LOG(LogStorageSocketBackend, Fatal, TEXT("Socket was not fully drained"));
+	}
+	Socket->Close();
+	delete Socket;
+}
+
+bool FStorageSocketConnectionBackend::InitializeInternal(TArrayView<const FString> InHostAddresses, int32 InPort)
+{
+	TArray<TSharedPtr<FInternetAddr>> HostAddresses = GetAddressFromString(SocketSubsystem, InHostAddresses, InPort);
+
+	if (!HostAddresses.Num())
+	{
+		UE_LOG(LogStorageSocketBackend, Fatal, TEXT("No valid Zen store host address specified"));
+		return false;
+	}
+
+	const int32 ServerVersion = HandshakeRequest(HostAddresses);
+	if (ServerVersion != 1)
+	{
+		return false;
+	}
+
+	return true;
+}
+
+int32 FStorageSocketConnectionBackend::HandshakeRequest(TArrayView<const TSharedPtr<FInternetAddr>> HostAddresses)
+{
+	TAnsiStringBuilder<256> ResourceBuilder;
+	ResourceBuilder.Append(OplogPath);
+
+	TArray<TSharedPtr<FInternetAddr>> SortedAddresses;
+	SortHostAddressesByLocalSubnet(HostAddresses, SortedAddresses);
+
+	for (const TSharedPtr<FInternetAddr>& Addr : SortedAddresses)
+	{
+		Hostname.Reset();
+		Hostname.Append(TCHAR_TO_ANSI(*Addr->ToString(false)));
+		ServerAddr = Addr;
+
+		UE_LOG(LogStorageSocketBackend, Display, TEXT("Trying to handshake with Zen at '%s'"), *Addr->ToString(true));
+
+		// Handshakes are done with a limited connection timeout so that we can find out if the destination is unreachable
+		// in a timely manner.
+		const float ConnectionTimeoutSeconds = 5.0f;
+		IStorageConnectionSocket* ConnectSocket = AcquireNewSocket(ConnectionTimeoutSeconds);
+		if (!ConnectSocket)
+		{
+			continue;
+		}
+		ReleaseSocket(ConnectSocket, true);
+
+		FStorageServerRequest Request("GET", *ResourceBuilder, Hostname);
+		if (IStorageConnectionSocket* Socket = Request.Send(Owner, false))
+		{
+			FStorageServerResponse Response(Owner, *Socket);
+
+			if (Response.IsOk())
+			{
+				FCbObject ResponseObj = Response.GetResponseObject();
+
+				// we currently don't have any concept of protocol versioning, if
+				// we succeed in communicating with the endpoint we're good since
+				// any breaking API change would need to be done in a backward
+				// compatible manner
+
+				return 1;
+			}
+			else
+			{
+				UE_LOG(LogStorageSocketBackend, Fatal, TEXT("Failed to handshake with Zen at %s. '%s'"), *ServerAddr->ToString(true), *Response.GetErrorMessage());
+			}
+		}
+		else
+		{
+			UE_LOG(LogStorageSocketBackend, Warning, TEXT("Failed to send handshake request to Zen at %s."), *ServerAddr->ToString(true));
+		}
+	}
+
+	UE_LOG(LogStorageSocketBackend, Fatal, TEXT("Failed to handshake with Zen at any of host addresses."));
+
+	Hostname.Reset();
+	ServerAddr.Reset();
+
+	return -1;
+}
+
+void FStorageSocketConnectionBackend::SortHostAddressesByLocalSubnet(TArrayView<const TSharedPtr<FInternetAddr>> HostAddresses, TArray<TSharedPtr<FInternetAddr>>& SortedHostAddresses)
+{
+	//no sorting needed cases
+	if (HostAddresses.Num() == 0)
+		return;
+
+	if (HostAddresses.Num() == 1)
+	{
+		SortedHostAddresses.Push(HostAddresses[0]);
+		return;
+	}
+
+	//Sorting logic:
+	//1 only on desktop, if it's an IPV6 address loopback (ends with ":1")
+	//2 only on desktop, if it's and IPV4 address loopback (starts with "127.0.0")
+	//3 if the host IPV4 subnet match the client subnet (xxx.xxx.xxx)
+	//4 remaining addresses
+	bool bCanBindAll = false;
+	bool bAppendPort = false;
+	TSharedPtr<FInternetAddr> localAddr = SocketSubsystem.GetLocalHostAddr(*GLog, bCanBindAll);
+	FString localAddrStringSubnet = localAddr->ToString(bAppendPort);
+
+	int32 localLastDotPos = INDEX_NONE;
+	if (localAddrStringSubnet.FindLastChar(TEXT('.'), localLastDotPos))
+	{
+		localAddrStringSubnet = localAddrStringSubnet.LeftChop(localAddrStringSubnet.Len() - localLastDotPos);
+	}
+
+	TArray<TSharedPtr<FInternetAddr>> IPV6Loopback;
+	TArray<TSharedPtr<FInternetAddr>> IPV4Loopback;
+	TArray<TSharedPtr<FInternetAddr>> RegularAddresses;
+
+	for (const TSharedPtr<FInternetAddr>& Addr : HostAddresses)
+	{
+		FString tempAddrStringSubnet = Addr->ToString(bAppendPort);
+
+#if PLATFORM_DESKTOP
+		if (Addr->GetProtocolType() == FNetworkProtocolTypes::IPv6)
+		{
+			if (tempAddrStringSubnet.EndsWith(":1"))
+			{
+				IPV6Loopback.Push(Addr);
+				continue;
+			}
+		}
+		else
+		{
+			if (tempAddrStringSubnet.StartsWith("127.0.0."))
+			{
+				IPV4Loopback.Push(Addr);
+				continue;
+			}
+		}
+#endif
+		int32 LastDotPos = INDEX_NONE;
+		if (tempAddrStringSubnet.FindLastChar(TEXT('.'), LastDotPos))
+		{
+			tempAddrStringSubnet = tempAddrStringSubnet.LeftChop(tempAddrStringSubnet.Len() - LastDotPos);
+		}
+
+		if (localAddrStringSubnet.Equals(tempAddrStringSubnet))
+			RegularAddresses.Insert(Addr, 0);
+		else
+			RegularAddresses.Push(Addr);
+	}
+
+
+	for (const TSharedPtr<FInternetAddr>& Addrv6lb : IPV6Loopback)
+	{
+		SortedHostAddresses.Push(Addrv6lb);
+	}
+
+	for (const TSharedPtr<FInternetAddr>& Addrv4lb : IPV4Loopback)
+	{
+		SortedHostAddresses.Push(Addrv4lb);
+	}
+
+	for (const TSharedPtr<FInternetAddr>& RegularAddr : RegularAddresses)
+	{
+		SortedHostAddresses.Push(RegularAddr);
+	}
+}
\ No newline at end of file
diff --git a/Engine/Source/Runtime/StorageServerClient/Private/StorageSocketConnectionBackend.h b/Engine/Source/Runtime/StorageServerClient/Private/StorageSocketConnectionBackend.h
new file mode 100644
--- /dev/null
+++ b/Engine/Source/Runtime/StorageServerClient/Private/StorageSocketConnectionBackend.h
@@ -0,0 +1,50 @@
+// Copyright Epic Games, Inc. All Rights Reserved.
+
+#pragma once
+
+#include "StorageServerConnection.h"
+
+class FStorageConnectionSocketFSocket : public IStorageConnectionSocket
+{
+public:
+	explicit FStorageConnectionSocketFSocket(FSocket* InSocket);
+	~FStorageConnectionSocketFSocket() override;
+
+	bool Send(const uint8* Data, const uint64 DataSize) override;
+	bool Recv(uint8* Data, const uint64 DataSize, uint64& BytesRead, ESocketReceiveFlags::Type ReceiveFlags) override;
+	bool HasPendingData(uint64& PendingDataSize) const override;
+	void Close() override;
+
+private:
+	FSocket* Socket;
+};
+
+class FStorageSocketConnectionBackend : public FStorageConnectionBackend
+{
+public:
+	FStorageSocketConnectionBackend(FStorageServerConnection& InOwner, ISocketSubsystem& InSocketSubsystem)
+		: FStorageConnectionBackend(InOwner)
+		, SocketSubsystem(InSocketSubsystem)
+	{
+	}
+
+	~FStorageSocketConnectionBackend() override;
+
+	IStorageConnectionSocket* AcquireSocketFromPool() override;
+	IStorageConnectionSocket* AcquireNewSocket(float TimeoutSeconds) override;
+	void ReleaseSocket(IStorageConnectionSocket* Socket, bool bKeepAlive) override;
+
+private:
+	bool InitializeInternal(TArrayView<const FString> InHostAddresses, int32 InPort) override;
+
+private:
+	int32 HandshakeRequest(TArrayView<const TSharedPtr<FInternetAddr>> HostAddresses);
+	void SortHostAddressesByLocalSubnet(TArrayView<const TSharedPtr<FInternetAddr>> HostAddresses, TArray<TSharedPtr<FInternetAddr>>& SortedHostAddresses);
+
+private:
+	ISocketSubsystem& SocketSubsystem;
+	TSharedPtr<FInternetAddr> ServerAddr;
+	TAnsiStringBuilder<1024> Hostname;
+	TArray<IStorageConnectionSocket*> SocketPool;
+	FCriticalSection SocketPoolCS;
+};
\ No newline at end of file
diff --git a/Engine/Source/Runtime/StorageServerClient/Public/StorageServerConnection.h b/Engine/Source/Runtime/StorageServerClient/Public/StorageServerConnection.h
--- a/Engine/Source/Runtime/StorageServerClient/Public/StorageServerConnection.h
+++ b/Engine/Source/Runtime/StorageServerClient/Public/StorageServerConnection.h
@@ -8,6 +8,7 @@
 #include "Misc/StringBuilder.h"
 #include "Memory/MemoryFwd.h"
 #include "Templates/SharedPointer.h"
+#include "SocketTypes.h"
 
 #if !UE_BUILD_SHIPPING
 
@@ -17,6 +18,10 @@ class FInternetAddr;
 class FSocket;
 class FStorageServerChunkBatchRequest;
 class FStorageServerConnection;
+class FStorageConnectionBackend;
+class FStorageSocketConnectionBackend;
+class FStorageServerPlatformConnectionBackend;
+class IStorageConnectionSocket;
 class ISocketSubsystem;
 struct FPackageStoreEntryResource;
 
@@ -75,6 +80,8 @@ class FStorageServerRequest
 	EStorageServerContentType AcceptContentType() const;
 protected:
 	friend FStorageServerConnection;
+	friend FStorageSocketConnectionBackend;
+	friend FStorageServerPlatformConnectionBackend;
 
 	FStorageServerRequest(
 		FAnsiStringView Verb,
@@ -82,7 +89,7 @@ protected:
 		FAnsiStringView Hostname,
 		EStorageServerContentType Accept = EStorageServerContentType::Binary);
 
-	FSocket* Send(FStorageServerConnection& Owner, bool bLogOnError = true);
+	IStorageConnectionSocket* Send(FStorageServerConnection& Owner, bool bLogOnError = true);
 	virtual void Serialize(void* V, int64 Length) override;
 
 	EStorageServerContentType AcceptType;
@@ -143,13 +150,15 @@ public:
 
 private:
 	friend FStorageServerConnection;
+	friend FStorageSocketConnectionBackend;
+	friend FStorageServerPlatformConnectionBackend;
 	friend FStorageServerChunkBatchRequest;
 
-	FStorageServerResponse(FStorageServerConnection& Owner, FSocket& Socket);
+	FStorageServerResponse(FStorageServerConnection& Owner, IStorageConnectionSocket& Socket);
 	void ReleaseSocket(bool bKeepAlive);
 
 	FStorageServerConnection& Owner;
-	FSocket* Socket = nullptr;
+	IStorageConnectionSocket* Socket = nullptr;
 	int64 ContentLength = 0;
 	int64 Position = 0;
 	int32 ErrorCode;
@@ -174,6 +183,40 @@ private:
 	int32 ChunkCountOffset = 0;
 };
 
+class IStorageConnectionSocket
+{
+public:
+	IStorageConnectionSocket() = default;
+	IStorageConnectionSocket(const IStorageConnectionSocket&) = delete;
+	virtual ~IStorageConnectionSocket() = default;
+
+	virtual bool Send(const uint8* Data, const uint64 DataSize)                                                      = 0;
+	virtual bool Recv(uint8* Data, const uint64 DataSize, uint64& BytesRead, ESocketReceiveFlags::Type ReceiveFlags) = 0;
+	virtual bool HasPendingData(uint64& PendingDataSize) const                                                       = 0;
+	virtual void Close()                                                                                             = 0;
+};
+
+class FStorageConnectionBackend
+{
+public:
+	FStorageConnectionBackend(FStorageServerConnection& InOwner);
+	virtual ~FStorageConnectionBackend() = default;
+
+	bool Initialize(TArrayView<const FString> InHostAddresses, int32 InPort, const TCHAR* InProjectNameOverride, const TCHAR* InPlatformNameOverride);
+
+	virtual IStorageConnectionSocket* AcquireSocketFromPool() = 0;
+	virtual IStorageConnectionSocket* AcquireNewSocket(float TimeoutSeconds = -1.f) = 0;
+	virtual void ReleaseSocket(IStorageConnectionSocket* Socket, bool bKeepAlive) = 0;
+	
+protected:
+	void InitOplog(const TCHAR* InProjectNameOverride, const TCHAR* InPlatformNameOverride);
+	virtual bool InitializeInternal(TArrayView<const FString> InHostAddresses, int32 Port) = 0;
+
+protected:
+	FStorageServerConnection& Owner;
+	TAnsiStringBuilder<1024> OplogPath;
+};
+
 class FStorageServerConnection
 {
 public:
@@ -195,18 +238,17 @@ private:
 	friend FStorageServerResponse;
 	friend FStorageServerChunkBatchRequest;
 
-	void SortHostAddressesByLocalSubnet(TArrayView<const TSharedPtr<FInternetAddr>> HostAddresses, TArray<TSharedPtr<FInternetAddr>>& SortedHostAddresses);
-	int32 HandshakeRequest(TArrayView<const TSharedPtr<FInternetAddr>> HostAddresses);
-	FSocket* AcquireSocketFromPool();
-	FSocket* AcquireNewSocket(float TimeoutSeconds = -1.0f);
-	void ReleaseSocket(FSocket* Socket, bool bKeepAlive);
+	bool CreateConnectionBackend(TArrayView<const FString> HostAddresses, int32 Port);
+	bool CreatePlatformBackend(const FString& HostAddresses, int32 Port);
+
+	FStorageConnectionBackend* GetConnectionBackend() const;
 
 	ISocketSubsystem& SocketSubsystem;
 	TAnsiStringBuilder<1024> OplogPath;
 	TSharedPtr<FInternetAddr> ServerAddr;
 	TAnsiStringBuilder<1024> Hostname;
-	FCriticalSection SocketPoolCritical;
-	TArray<FSocket*> SocketPool;
+	TUniquePtr<FStorageConnectionBackend> ConnectionBackend;
+	FCriticalSection BackendCS;
 };
 
 #endif
