Compilation fix - Zen platform transport backend

diff --git a/Engine/Source/Runtime/StorageServerClient/Private/StoragePlatformConnectionBackend.cpp b/Engine/Source/Runtime/StorageServerClient/Private/StoragePlatformConnectionBackend.cpp
--- a/Engine/Source/Runtime/StorageServerClient/Private/StoragePlatformConnectionBackend.cpp
+++ b/Engine/Source/Runtime/StorageServerClient/Private/StoragePlatformConnectionBackend.cpp
@@ -6,6 +6,8 @@
 #include "Serialization/CompactBinarySerialization.h"
 #include "HAL/PlatformProcess.h"
 
+#if !UE_BUILD_SHIPPING
+
 DEFINE_LOG_CATEGORY_STATIC(LogStorageServerPlatformBackend, Log, All);
 
 FStorageConnectionPlatformSocket::FStorageConnectionPlatformSocket(IPlatformHostCommunication* HostCommunication, IPlatformHostSocketPtr InSocket, int32 InProtocolNumber)
@@ -252,3 +254,5 @@ int32 FStorageServerPlatformConnectionBackend::HandshakeRequest(TArrayView<const
 
 	return -1;
 }
+
+#endif // !UE_BUILD_SHIPPING
diff --git a/Engine/Source/Runtime/StorageServerClient/Private/StoragePlatformConnectionBackend.h b/Engine/Source/Runtime/StorageServerClient/Private/StoragePlatformConnectionBackend.h
--- a/Engine/Source/Runtime/StorageServerClient/Private/StoragePlatformConnectionBackend.h
+++ b/Engine/Source/Runtime/StorageServerClient/Private/StoragePlatformConnectionBackend.h
@@ -7,6 +7,8 @@
 #include "GenericPlatform/GenericPlatformHostCommunication.h"
 #include "GenericPlatform/GenericPlatformHostSocket.h"
 
+#if !UE_BUILD_SHIPPING
+
 class FConnectionCircularBuffer
 {
 public:
@@ -170,4 +172,6 @@ private:
 	TBitArray<> UsedSockets; // bitset to keep track of used sockets in the pool
 
 	TAnsiStringBuilder<1024> Hostname;
-};
\ No newline at end of file
+};
+
+#endif // !UE_BUILD_SHIPPING
\ No newline at end of file
diff --git a/Engine/Source/Runtime/StorageServerClient/Private/StorageSocketConnectionBackend.cpp b/Engine/Source/Runtime/StorageServerClient/Private/StorageSocketConnectionBackend.cpp
--- a/Engine/Source/Runtime/StorageServerClient/Private/StorageSocketConnectionBackend.cpp
+++ b/Engine/Source/Runtime/StorageServerClient/Private/StorageSocketConnectionBackend.cpp
@@ -9,6 +9,8 @@
 #include "SocketSubsystem.h"
 #include "Sockets.h"
 
+#if !UE_BUILD_SHIPPING
+
 DEFINE_LOG_CATEGORY_STATIC(LogStorageSocketBackend, Log, All);
 
 static TArray<TSharedPtr<FInternetAddr>> GetAddressFromString(ISocketSubsystem& SocketSubsystem, TArrayView<const FString> HostAddresses, const int32 Port)
@@ -330,4 +332,6 @@ void FStorageSocketConnectionBackend::SortHostAddressesByLocalSubnet(TArrayView<
 	{
 		SortedHostAddresses.Push(RegularAddr);
 	}
-}
\ No newline at end of file
+}
+
+#endif // !UE_BUILD_SHIPPING
\ No newline at end of file
diff --git a/Engine/Source/Runtime/StorageServerClient/Private/StorageSocketConnectionBackend.h b/Engine/Source/Runtime/StorageServerClient/Private/StorageSocketConnectionBackend.h
--- a/Engine/Source/Runtime/StorageServerClient/Private/StorageSocketConnectionBackend.h
+++ b/Engine/Source/Runtime/StorageServerClient/Private/StorageSocketConnectionBackend.h
@@ -4,6 +4,8 @@
 
 #include "StorageServerConnection.h"
 
+#if !UE_BUILD_SHIPPING
+
 class FStorageConnectionSocketFSocket : public IStorageConnectionSocket
 {
 public:
@@ -47,4 +49,6 @@ private:
 	TAnsiStringBuilder<1024> Hostname;
 	TArray<IStorageConnectionSocket*> SocketPool;
 	FCriticalSection SocketPoolCS;
-};
\ No newline at end of file
+};
+
+#endif // !UE_BUILD_SHIPPING
\ No newline at end of file
