Adding an error msg in case we find but fail to deserialize ue.projectstore

diff --git a/Engine/Source/Runtime/StorageServerClient/Private/StorageServerPlatformFile.cpp b/Engine/Source/Runtime/StorageServerClient/Private/StorageServerPlatformFile.cpp
--- a/Engine/Source/Runtime/StorageServerClient/Private/StorageServerPlatformFile.cpp
+++ b/Engine/Source/Runtime/StorageServerClient/Private/StorageServerPlatformFile.cpp
@@ -343,6 +343,10 @@ bool FStorageServerPlatformFile::ShouldBeUsed(IPlatformFile* Inner, const TCHAR*
 				UE_LOG(LogStorageServerPlatformFile, Display, TEXT("Using connection settings from ue.projectstore: HostAddrs='%s' and HostPort='%d'"), *FString::Join(HostAddrs, TEXT("+")), HostPort);
 			}
 		}
+		else
+		{
+			UE_LOG(LogStorageServerPlatformFile, Error, TEXT("Failed to Deserialize ue.projectstore!'"));
+		}
 	}
 
 	FString Host;
