Fixing vulkan PSO cache when using Zen streaming

diff --git a/Engine/Source/Runtime/StorageServerClient/Private/StorageServerPlatformFile.cpp b/Engine/Source/Runtime/StorageServerClient/Private/StorageServerPlatformFile.cpp
--- a/Engine/Source/Runtime/StorageServerClient/Private/StorageServerPlatformFile.cpp
+++ b/Engine/Source/Runtime/StorageServerClient/Private/StorageServerPlatformFile.cpp
@@ -627,6 +627,11 @@ bool FStorageServerPlatformFile::IterateDirectoryStat(const TCHAR* Directory, FD
 	return bResult;
 }
 
+IMappedFileHandle* FStorageServerPlatformFile::OpenMapped(const TCHAR* Filename)
+{
+	return LowerLevel->OpenMapped(Filename);
+}
+
 bool FStorageServerPlatformFile::DirectoryExists(const TCHAR* Directory)
 {
 	TStringBuilder<1024> StorageServerDirectory;
diff --git a/Engine/Source/Runtime/StorageServerClient/Private/StorageServerPlatformFile.h b/Engine/Source/Runtime/StorageServerClient/Private/StorageServerPlatformFile.h
--- a/Engine/Source/Runtime/StorageServerClient/Private/StorageServerPlatformFile.h
+++ b/Engine/Source/Runtime/StorageServerClient/Private/StorageServerPlatformFile.h
@@ -88,6 +88,7 @@ public:
 	virtual FFileStatData GetStatData(const TCHAR* FilenameOrDirectory) override;
 	virtual bool IterateDirectory(const TCHAR* Directory, FDirectoryVisitor& Visitor) override;
 	virtual bool IterateDirectoryStat(const TCHAR* Directory, FDirectoryStatVisitor& Visitor) override;
+	virtual IMappedFileHandle* OpenMapped(const TCHAR* Filename) override;
 	virtual FString GetFilenameOnDisk(const TCHAR* Filename) override;
 	virtual bool DeleteFile(const TCHAR* Filename) override;
 	virtual bool MoveFile(const TCHAR* To, const TCHAR* From) override;
