In support of Android streaming workflows:
- Allow android to connect to non-external (local) addresses like desktop platforms do because of the way in which ADB reverse port mapping is setup for Android.
-Ensure that on any platform where we are using BSD sockets, but haven't enabled IPV6 BSD socket support (which Android seems to be right now) that we make an effort to translate IPV6 loopback to IPV4 loopback addresses and warn if we encounter other IPV6 addresses.

diff --git a/Engine/Source/Runtime/StorageServerClient/Private/StorageServerPlatformFile.cpp b/Engine/Source/Runtime/StorageServerClient/Private/StorageServerPlatformFile.cpp
--- a/Engine/Source/Runtime/StorageServerClient/Private/StorageServerPlatformFile.cpp
+++ b/Engine/Source/Runtime/StorageServerClient/Private/StorageServerPlatformFile.cpp
@@ -316,7 +316,7 @@ bool FStorageServerPlatformFile::ShouldBeUsed(IPlatformFile* Inner, const TCHAR*
 			if (ProjectStoreObject->TryGetObjectField(TEXT("zenserver"), ZenServerObjectPtr) && (ZenServerObjectPtr != nullptr))
 			{
 				const TSharedPtr<FJsonObject>& ZenServerObject = *ZenServerObjectPtr;
-#if PLATFORM_DESKTOP
+#if PLATFORM_DESKTOP || PLATFORM_ANDROID
 				FString HostName;
 				if (ZenServerObject->TryGetStringField(TEXT("hostname"), HostName) && !HostName.IsEmpty())
 				{
diff --git a/Engine/Source/Runtime/StorageServerClient/Private/StorageSocketConnectionBackend.cpp b/Engine/Source/Runtime/StorageServerClient/Private/StorageSocketConnectionBackend.cpp
--- a/Engine/Source/Runtime/StorageServerClient/Private/StorageSocketConnectionBackend.cpp
+++ b/Engine/Source/Runtime/StorageServerClient/Private/StorageSocketConnectionBackend.cpp
@@ -24,7 +24,21 @@ static TArray<TSharedPtr<FInternetAddr>> GetAddressFromString(ISocketSubsystem&
 		const FString* EffectiveHostAddr = &HostAddr;
 		if (!HostAddr.IsEmpty() && HostAddr[0] == TEXT('[') && HostAddr[HostAddr.Len() - 1] == TEXT(']'))
 		{
+#if PLATFORM_HAS_BSD_SOCKETS && !PLATFORM_HAS_BSD_IPV6_SOCKETS
+			// If the platform doesn't have IPV6 BSD Sockets, then handle an attempt at conversion of loopback addresses, and skip and warn about other addresses
+			if (HostAddr == TEXT("[::1]"))
+			{
+				// Substitute IPV4 loopback for IPV6 loopback
+				ModifiedHostAddr = TEXT("127.0.0.1");
+			}
+			else
+			{
+				UE_LOG(LogStorageSocketBackend, Warning, TEXT("Ignoring storage server host IPV6 address on platform that doesn't support IPV6: %s"), *HostAddr);
+				continue;
+			}
+#else
 			ModifiedHostAddr = FStringView(HostAddr).Mid(1, HostAddr.Len() - 2);
+#endif
 			EffectiveHostAddr = &ModifiedHostAddr;
 		}
 		TSharedPtr<FInternetAddr> Addr = SocketSubsystem.GetAddressFromString(*EffectiveHostAddr);
