Add an exclusion list of file extensions that we don't forward to lower level physical platform file layer when using the StorageServerPlatformFile wrapper.

diff --git a/Engine/Source/Runtime/StorageServerClient/Private/StorageServerPlatformFile.cpp b/Engine/Source/Runtime/StorageServerClient/Private/StorageServerPlatformFile.cpp
--- a/Engine/Source/Runtime/StorageServerClient/Private/StorageServerPlatformFile.cpp
+++ b/Engine/Source/Runtime/StorageServerClient/Private/StorageServerPlatformFile.cpp
@@ -11,6 +11,7 @@
 #include "Misc/CoreDelegates.h"
 #include "Misc/FileHelper.h"
 #include "Misc/Paths.h"
+#include "Misc/PathViews.h"
 #include "Misc/ScopeRWLock.h"
 #include "Misc/StringBuilder.h"
 #include "Modules/ModuleManager.h"
@@ -26,6 +27,10 @@ DEFINE_LOG_CATEGORY_STATIC(LogStorageServerPlatformFile, Log, All);
 
 #if !UE_BUILD_SHIPPING
 
+#ifndef EXCLUDE_NONSERVER_UE_EXTENSIONS
+#define EXCLUDE_NONSERVER_UE_EXTENSIONS 1	// Use .Build.cs file to disable this if the game relies on accessing loose files on the local filesystem
+#endif
+
 FStorageServerFileSystemTOC::~FStorageServerFileSystemTOC()
 {
 	FWriteScopeLock _(TocLock);
@@ -370,6 +375,15 @@ bool FStorageServerPlatformFile::Initialize(IPlatformFile* Inner, const TCHAR* C
 	LowerLevel = Inner;
 	if (HostAddrs.Num() > 0)
 	{
+#if EXCLUDE_NONSERVER_UE_EXTENSIONS && !WITH_EDITOR
+		// Extensions for file types that should only ever be on the server. Used to stop unnecessary access to the lower level platform file.
+		ExcludedNonServerExtensions.Add(TEXT("uasset"));
+		ExcludedNonServerExtensions.Add(TEXT("umap"));
+		ExcludedNonServerExtensions.Add(TEXT("ubulk"));
+		ExcludedNonServerExtensions.Add(TEXT("uexp"));
+		ExcludedNonServerExtensions.Add(TEXT("uptnl"));
+		ExcludedNonServerExtensions.Add(TEXT("ushaderbytecode"));
+#endif
 		// Don't initialize the connection yet because we want to incorporate project file path information into the initialization.
 
 		TUniquePtr<FArchive> ProjectStoreMarkerReader = TryFindProjectStoreMarkerFile(Inner);
@@ -469,7 +483,8 @@ bool FStorageServerPlatformFile::FileExists(const TCHAR* Filename)
 	{
 		return true;
 	}
-	return LowerLevel->FileExists(Filename);
+
+	return IsNonServerFilenameAllowed(Filename) ? LowerLevel->FileExists(Filename) : false;
 }
 
 FDateTime FStorageServerPlatformFile::GetTimeStamp(const TCHAR* Filename)
@@ -484,7 +499,7 @@ FDateTime FStorageServerPlatformFile::GetTimeStamp(const TCHAR* Filename)
 			return FileStatData.ModificationTime;
 		}
 	}
-	return LowerLevel->GetTimeStamp(Filename);
+	return IsNonServerFilenameAllowed(Filename) ? LowerLevel->GetTimeStamp(Filename) : FDateTime::MinValue();
 }
 
 FDateTime FStorageServerPlatformFile::GetAccessTimeStamp(const TCHAR* Filename)
@@ -499,7 +514,7 @@ FDateTime FStorageServerPlatformFile::GetAccessTimeStamp(const TCHAR* Filename)
 			return FileStatData.AccessTime;
 		}
 	}
-	return LowerLevel->GetAccessTimeStamp(Filename);
+	return IsNonServerFilenameAllowed(Filename) ? LowerLevel->GetAccessTimeStamp(Filename) : FDateTime::MinValue();
 }
 
 int64 FStorageServerPlatformFile::FileSize(const TCHAR* Filename)
@@ -514,7 +529,7 @@ int64 FStorageServerPlatformFile::FileSize(const TCHAR* Filename)
 			return FileStatData.FileSize;
 		}
 	}
-	return LowerLevel->FileSize(Filename);
+	return IsNonServerFilenameAllowed(Filename) ? LowerLevel->FileSize(Filename) : -1;
 }
 
 int64 FStorageServerPlatformFile::FileSize(const TCHAR* Filename, EPlatformFileFlags PlatformFlags)
@@ -535,7 +550,7 @@ int64 FStorageServerPlatformFile::FileSize(const TCHAR* Filename, EPlatformFileF
 		return -1;
 	}
 
-	return LowerLevel->FileSize(Filename);
+	return IsNonServerFilenameAllowed(Filename) ? LowerLevel->FileSize(Filename) : -1;
 }
 
 
@@ -546,7 +561,7 @@ bool FStorageServerPlatformFile::IsReadOnly(const TCHAR* Filename)
 	{
 		return true;
 	}
-	return LowerLevel->IsReadOnly(Filename);
+	return IsNonServerFilenameAllowed(Filename) ? LowerLevel->IsReadOnly(Filename) : false;
 }
 
 FFileStatData FStorageServerPlatformFile::GetStatData(const TCHAR* FilenameOrDirectory)
@@ -569,7 +584,12 @@ FFileStatData FStorageServerPlatformFile::GetStatData(const TCHAR* FilenameOrDir
 				true);
 		}
 	}
-	return LowerLevel->GetStatData(FilenameOrDirectory);
+	FFileStatData FileStatData;
+	if (IsNonServerFilenameAllowed(FilenameOrDirectory))
+	{
+		FileStatData = LowerLevel->GetStatData(FilenameOrDirectory);
+	}
+	return FileStatData;
 }
 
 IFileHandle* FStorageServerPlatformFile::InternalOpenFile(const FIoChunkId& FileChunkId, const TCHAR* LocalFilename)
@@ -587,7 +607,7 @@ IFileHandle* FStorageServerPlatformFile::OpenRead(const TCHAR* Filename, bool bA
 			return InternalOpenFile(*FileChunkId, Filename);
 		}
 	}
-	return LowerLevel->OpenRead(Filename, bAllowWrite);
+	return IsNonServerFilenameAllowed(Filename) ? LowerLevel->OpenRead(Filename, bAllowWrite) : nullptr;
 }
 
 bool FStorageServerPlatformFile::IterateDirectory(const TCHAR* Directory, IPlatformFile::FDirectoryVisitor& Visitor)
@@ -651,7 +671,7 @@ bool FStorageServerPlatformFile::IterateDirectoryStat(const TCHAR* Directory, FD
 
 IMappedFileHandle* FStorageServerPlatformFile::OpenMapped(const TCHAR* Filename)
 {
-	return LowerLevel->OpenMapped(Filename);
+	return IsNonServerFilenameAllowed(Filename) ? LowerLevel->OpenMapped(Filename) : nullptr;
 }
 
 bool FStorageServerPlatformFile::DirectoryExists(const TCHAR* Directory)
@@ -672,7 +692,7 @@ FString FStorageServerPlatformFile::GetFilenameOnDisk(const TCHAR* Filename)
 		UE_LOG(LogStorageServerPlatformFile, Warning, TEXT("Attempting to get disk filename of remote file '%s'"), Filename);
 		return Filename;
 	}
-	return LowerLevel->GetFilenameOnDisk(Filename);
+	return IsNonServerFilenameAllowed(Filename) ? LowerLevel->GetFilenameOnDisk(Filename) : Filename;
 }
 
 bool FStorageServerPlatformFile::DeleteFile(const TCHAR* Filename)
@@ -827,6 +847,24 @@ FString FStorageServerPlatformFile::ConvertToAbsolutePathForExternalAppForRead(c
 	return LowerLevel->ConvertToAbsolutePathForExternalAppForRead(Filename);
 }
 
+bool FStorageServerPlatformFile::IsNonServerFilenameAllowed(FStringView InFilename)
+{
+	bool bAllowed = true;
+
+#if EXCLUDE_NONSERVER_UE_EXTENSIONS
+	if (!HostAddrs.IsEmpty() && (LowerLevel == &IPlatformFile::GetPlatformPhysical()))
+	{
+		FName Ext = FName(FPathViews::GetExtension(InFilename));
+		bAllowed = !ExcludedNonServerExtensions.Contains(Ext);
+		UE_CLOG(!bAllowed, LogStorageServerPlatformFile, VeryVerbose,
+			TEXT("Access to file '%.*s' is limited to server contents due to file extension being listed in ExcludedNonServerExtensions."),
+			InFilename.Len(), InFilename.GetData())
+	}
+#endif
+
+	return bAllowed;
+}
+
 bool FStorageServerPlatformFile::MakeStorageServerPath(const TCHAR* LocalFilenameOrDirectory, FStringBuilderBase& OutPath) const
 {
 	FStringView LocalEngineDirView(FPlatformMisc::EngineDir());
diff --git a/Engine/Source/Runtime/StorageServerClient/Private/StorageServerPlatformFile.h b/Engine/Source/Runtime/StorageServerClient/Private/StorageServerPlatformFile.h
--- a/Engine/Source/Runtime/StorageServerClient/Private/StorageServerPlatformFile.h
+++ b/Engine/Source/Runtime/StorageServerClient/Private/StorageServerPlatformFile.h
@@ -104,6 +104,7 @@ public:
 private:
 	friend class FStorageServerFileHandle;
 
+	bool IsNonServerFilenameAllowed(FStringView InFilename);
 	bool MakeStorageServerPath(const TCHAR* LocalFilenameOrDirectory, FStringBuilderBase& OutPath) const;
 	bool MakeLocalPath(const TCHAR* ServerFilenameOrDirectory, FStringBuilderBase& OutPath) const;
 	IFileHandle* InternalOpenFile(const FIoChunkId& FileChunkId, const TCHAR* LocalFilename);
@@ -115,6 +116,7 @@ private:
 #endif
 	TUniquePtr<FArchive> TryFindProjectStoreMarkerFile(IPlatformFile* Inner) const;
 
+	TSet<FName> ExcludedNonServerExtensions;
 	IPlatformFile* LowerLevel = nullptr;
 	FStringView ServerEngineDirView = FStringView(TEXT("/{engine}/"));
 	FStringView ServerProjectDirView = FStringView(TEXT("/{project}/"));
