When client is streaming data from ZenServer, display a warning on screen: "ZenServer streaming from <IP>"

diff --git a/Engine/Source/Runtime/StorageServerClient/Private/StoragePlatformConnectionBackend.h b/Engine/Source/Runtime/StorageServerClient/Private/StoragePlatformConnectionBackend.h
--- a/Engine/Source/Runtime/StorageServerClient/Private/StoragePlatformConnectionBackend.h
+++ b/Engine/Source/Runtime/StorageServerClient/Private/StoragePlatformConnectionBackend.h
@@ -157,6 +157,11 @@ public:
 	IStorageConnectionSocket* AcquireSocketFromPool() override;
 	void ReleaseSocket(IStorageConnectionSocket* Socket, bool bKeepAlive) override;
 
+	FString GetHostName() override
+	{
+		return Hostname.ToString();
+	}
+
 protected:
 	bool InitializeInternal(TArrayView<const FString> InHostAddresses, int32 Port) override;
 
diff --git a/Engine/Source/Runtime/StorageServerClient/Private/StorageServerConnection.cpp b/Engine/Source/Runtime/StorageServerClient/Private/StorageServerConnection.cpp
--- a/Engine/Source/Runtime/StorageServerClient/Private/StorageServerConnection.cpp
+++ b/Engine/Source/Runtime/StorageServerClient/Private/StorageServerConnection.cpp
@@ -19,6 +19,7 @@
 #include "StorageSocketConnectionBackend.h"
 #include "StoragePlatformConnectionBackend.h"
 #include "GenericPlatform/GenericPlatformHostCommunication.h"
+#include "Engine/Engine.h"
 
 #if !UE_BUILD_SHIPPING
 
@@ -499,13 +500,14 @@ bool FStorageServerConnection::Initialize(TArrayView<const FString> InHostAddres
 	{
 		return false;
 	}
-
+	
 	return ConnectionBackend->Initialize(InHostAddresses, InPort, InProjectNameOverride, InPlatformNameOverride);
 }
 
 void FStorageServerConnection::PackageStoreRequest(TFunctionRef<void(FPackageStoreEntryResource&&)> Callback)
 {
 	TRACE_CPUPROFILER_EVENT_SCOPE(StorageServerPackageStoreRequest);
+	ShowDebugMessage();
 
 	TAnsiStringBuilder<256> ResourceBuilder;
 	ResourceBuilder.Append(OplogPath).Append("/entries?fieldfilter=packagestoreentry");
@@ -537,6 +539,8 @@ void FStorageServerConnection::PackageStoreRequest(TFunctionRef<void(FPackageSto
 
 void FStorageServerConnection::FileManifestRequest(TFunctionRef<void(FIoChunkId Id, FStringView Path)> Callback)
 {
+	ShowDebugMessage();
+
 	TAnsiStringBuilder<256> ResourceBuilder;
 	ResourceBuilder.Append(OplogPath).Append("/files?filter=client");
 	FStorageServerRequest Request("GET", *ResourceBuilder, Hostname, EStorageServerContentType::CbObject);
@@ -609,6 +613,8 @@ bool FStorageServerConnection::ReadChunkRequest(const FIoChunkId& ChunkId, uint6
 {
 	TRACE_CPUPROFILER_EVENT_SCOPE(ZenHttpClient::ReadChunkRequest);
 
+	ShowDebugMessage();
+
 	TAnsiStringBuilder<256> ResourceBuilder;
 	ResourceBuilder.Append(OplogPath) << "/" << ChunkId;
 
@@ -714,4 +720,17 @@ FStorageConnectionBackend* FStorageServerConnection::GetConnectionBackend() cons
 	return ConnectionBackend.Get();
 }
 
+void FStorageServerConnection::ShowDebugMessage()
+{
+	static bool bShowDebugMessage = true;
+
+	if ((bShowDebugMessage) && (GEngine))
+	{
+		FString ZenConnectionDebugMsg;
+		ZenConnectionDebugMsg = FString::Printf(TEXT("ZenServer streaming from %s"), *ConnectionBackend->GetHostName());
+		GEngine->AddOnScreenDebugMessage(INDEX_NONE, 86400.0f, FColor::White, ZenConnectionDebugMsg, false);
+		bShowDebugMessage = false;
+	}
+}
+
 #endif
diff --git a/Engine/Source/Runtime/StorageServerClient/Private/StorageSocketConnectionBackend.h b/Engine/Source/Runtime/StorageServerClient/Private/StorageSocketConnectionBackend.h
--- a/Engine/Source/Runtime/StorageServerClient/Private/StorageSocketConnectionBackend.h
+++ b/Engine/Source/Runtime/StorageServerClient/Private/StorageSocketConnectionBackend.h
@@ -36,6 +36,11 @@ public:
 	IStorageConnectionSocket* AcquireNewSocket(float TimeoutSeconds) override;
 	void ReleaseSocket(IStorageConnectionSocket* Socket, bool bKeepAlive) override;
 
+	FString GetHostName() override
+	{
+		return Hostname.ToString();
+	}
+
 private:
 	bool InitializeInternal(TArrayView<const FString> InHostAddresses, int32 InPort) override;
 
diff --git a/Engine/Source/Runtime/StorageServerClient/Public/StorageServerConnection.h b/Engine/Source/Runtime/StorageServerClient/Public/StorageServerConnection.h
--- a/Engine/Source/Runtime/StorageServerClient/Public/StorageServerConnection.h
+++ b/Engine/Source/Runtime/StorageServerClient/Public/StorageServerConnection.h
@@ -207,6 +207,7 @@ public:
 	virtual IStorageConnectionSocket* AcquireSocketFromPool() = 0;
 	virtual IStorageConnectionSocket* AcquireNewSocket(float TimeoutSeconds = -1.f) = 0;
 	virtual void ReleaseSocket(IStorageConnectionSocket* Socket, bool bKeepAlive) = 0;
+	virtual FString GetHostName() = 0;
 	
 protected:
 	void InitOplog(const TCHAR* InProjectNameOverride, const TCHAR* InPlatformNameOverride);
@@ -243,6 +244,8 @@ private:
 
 	FStorageConnectionBackend* GetConnectionBackend() const;
 
+	void ShowDebugMessage();
+
 	ISocketSubsystem& SocketSubsystem;
 	TAnsiStringBuilder<1024> OplogPath;
 	TSharedPtr<FInternetAddr> ServerAddr;
diff --git a/Engine/Source/Runtime/StorageServerClient/StorageServerClient.Build.cs b/Engine/Source/Runtime/StorageServerClient/StorageServerClient.Build.cs
--- a/Engine/Source/Runtime/StorageServerClient/StorageServerClient.Build.cs
+++ b/Engine/Source/Runtime/StorageServerClient/StorageServerClient.Build.cs
@@ -10,6 +10,7 @@ public class StorageServerClient : ModuleRules
             new string[] {
                 "Core",
                 "CoreUObject",
+				"Engine",        // only for on-screen debug messages
                 "Sockets",
                 "CookOnTheFly",
                 "Json"
