Ensure Android prioritizes loopback addresses when picking from possible storage server connections.

diff --git a/Engine/Source/Runtime/StorageServerClient/Private/StorageSocketConnectionBackend.cpp b/Engine/Source/Runtime/StorageServerClient/Private/StorageSocketConnectionBackend.cpp
--- a/Engine/Source/Runtime/StorageServerClient/Private/StorageSocketConnectionBackend.cpp
+++ b/Engine/Source/Runtime/StorageServerClient/Private/StorageSocketConnectionBackend.cpp
@@ -301,7 +301,7 @@ void FStorageSocketConnectionBackend::SortHostAddressesByLocalSubnet(TArrayView<
 	{
 		FString tempAddrStringSubnet = Addr->ToString(bAppendPort);
 
-#if PLATFORM_DESKTOP
+#if PLATFORM_DESKTOP || PLATFORM_ANDROID
 		if (Addr->GetProtocolType() == FNetworkProtocolTypes::IPv6)
 		{
 			if (tempAddrStringSubnet.EndsWith(":1"))
