Don't emit fatal verbosity messages when we've failed to connect to a storage server.  Allow the use of the more user-friendly message.

diff --git a/Engine/Source/Runtime/StorageServerClient/Private/StorageServerPlatformFile.cpp b/Engine/Source/Runtime/StorageServerClient/Private/StorageServerPlatformFile.cpp
--- a/Engine/Source/Runtime/StorageServerClient/Private/StorageServerPlatformFile.cpp
+++ b/Engine/Source/Runtime/StorageServerClient/Private/StorageServerPlatformFile.cpp
@@ -459,16 +459,16 @@ void FStorageServerPlatformFile::InitializeAfterProjectFilePath()
 	{
 		if (!FApp::IsUnattended())
 		{
-			FText FailedConnectionTitle = NSLOCTEXT("StorageServer", "StorageServer_ConnectFailedTitle", "Failed to connect");
-			FText FailedConnectionText = FText::Format(NSLOCTEXT("StorageServer", "StorageServer_ConnectFailedText",
-				"Network data streaming failed to connect to any of the following data sources:\n\n{0}\n\n"
+			FString FailedConnectionTitle = TEXT("Failed to connect");
+			FString FailedConnectionText = FString::Printf(TEXT(
+				"Network data streaming failed to connect to any of the following data sources:\n\n%s\n\n"
 				"This can be due to the sources being offline, the Unreal Zen Storage process not currently running, "
 				"invalid addresses, firewall blocking, or the sources being on a different network from this device. "
 				"Please verify that your Unreal Zen Storage process is running using the ZenDashboard utility. "
 				"If these issues can't be addressed, you can use an installed build without network data streaming by "
 				"building with the '-pak' argument. This process will now exit."),
-				FText::FromString(FString::Join(HostAddrs, TEXT("\n"))));
-			FPlatformMisc::MessageBoxExt(EAppMsgType::Ok, *FailedConnectionText.ToString(), *FailedConnectionTitle.ToString());
+				*FString::Join(HostAddrs, TEXT("\n")));
+			FPlatformMisc::MessageBoxExt(EAppMsgType::Ok, *FailedConnectionText, *FailedConnectionTitle);
 		}
 
 		UE_LOG(LogStorageServerPlatformFile, Error, TEXT("Failed to initialize connection to %s"), *FString::Join(HostAddrs, TEXT("\n")));
diff --git a/Engine/Source/Runtime/StorageServerClient/Private/StorageSocketConnectionBackend.cpp b/Engine/Source/Runtime/StorageServerClient/Private/StorageSocketConnectionBackend.cpp
--- a/Engine/Source/Runtime/StorageServerClient/Private/StorageSocketConnectionBackend.cpp
+++ b/Engine/Source/Runtime/StorageServerClient/Private/StorageSocketConnectionBackend.cpp
@@ -248,7 +248,7 @@ int32 FStorageSocketConnectionBackend::HandshakeRequest(TArrayView<const TShared
 			}
 			else
 			{
-				UE_LOG(LogStorageSocketBackend, Fatal, TEXT("Failed to handshake with Zen at %s. '%s'"), *ServerAddr->ToString(true), *Response.GetErrorMessage());
+				UE_LOG(LogStorageSocketBackend, Error, TEXT("Failed to handshake with Zen at %s. '%s'"), *ServerAddr->ToString(true), *Response.GetErrorMessage());
 			}
 		}
 		else
@@ -257,7 +257,7 @@ int32 FStorageSocketConnectionBackend::HandshakeRequest(TArrayView<const TShared
 		}
 	}
 
-	UE_LOG(LogStorageSocketBackend, Fatal, TEXT("Failed to handshake with Zen at any of host addresses."));
+	UE_LOG(LogStorageSocketBackend, Error, TEXT("Failed to handshake with Zen at any of host addresses."));
 
 	Hostname.Reset();
 	ServerAddr.Reset();
