Avoid doing round-trip stat messages for getting timestamps that are always going to be Now if the file exists when we already have a means on the client to determine file existence.

diff --git a/Engine/Source/Runtime/StorageServerClient/Private/StorageServerPlatformFile.cpp b/Engine/Source/Runtime/StorageServerClient/Private/StorageServerPlatformFile.cpp
--- a/Engine/Source/Runtime/StorageServerClient/Private/StorageServerPlatformFile.cpp
+++ b/Engine/Source/Runtime/StorageServerClient/Private/StorageServerPlatformFile.cpp
@@ -492,11 +492,9 @@ FDateTime FStorageServerPlatformFile::GetTimeStamp(const TCHAR* Filename)
 	TStringBuilder<1024> StorageServerFilename;
 	if (MakeStorageServerPath(Filename, StorageServerFilename))
 	{
-		if (const FIoChunkId* FileChunkId = ServerToc.GetFileChunkId(*StorageServerFilename))
+		if (ServerToc.FileExists(*StorageServerFilename))
 		{
-			const FFileStatData FileStatData = SendGetStatDataMessage(*FileChunkId);
-			check(FileStatData.bIsValid);
-			return FileStatData.ModificationTime;
+			return FDateTime::Now();
 		}
 	}
 	return IsNonServerFilenameAllowed(Filename) ? LowerLevel->GetTimeStamp(Filename) : FDateTime::MinValue();
@@ -507,11 +505,9 @@ FDateTime FStorageServerPlatformFile::GetAccessTimeStamp(const TCHAR* Filename)
 	TStringBuilder<1024> StorageServerFilename;
 	if (MakeStorageServerPath(Filename, StorageServerFilename))
 	{
-		if (const FIoChunkId* FileChunkId = ServerToc.GetFileChunkId(*StorageServerFilename))
+		if (ServerToc.FileExists(*StorageServerFilename))
 		{
-			const FFileStatData FileStatData = SendGetStatDataMessage(*FileChunkId);
-			check(FileStatData.bIsValid);
-			return FileStatData.AccessTime;
+			return FDateTime::Now();
 		}
 	}
 	return IsNonServerFilenameAllowed(Filename) ? LowerLevel->GetAccessTimeStamp(Filename) : FDateTime::MinValue();
