Adding throughput computation for ZenServer in StorageServerRequest on the client.
Adding CSV stats for ZenServer throughput and request count.
Adding real time debug graphs for ZenServer throughput and request count.

diff --git a/Engine/Source/Runtime/StorageServerClient/Private/DebugStorageServerConnection.cpp b/Engine/Source/Runtime/StorageServerClient/Private/DebugStorageServerConnection.cpp
new file mode 100644
--- /dev/null
+++ b/Engine/Source/Runtime/StorageServerClient/Private/DebugStorageServerConnection.cpp
@@ -0,0 +1,243 @@
+// Copyright Epic Games, Inc. All Rights Reserved.
+
+#include "DebugStorageServerConnection.h"
+#include "StorageServerConnection.h"
+
+#include "Debug/DebugDrawService.h"
+#include "Engine/GameEngine.h"
+#include "ProfilingDebugging/CsvProfiler.h"
+
+#if CSV_PROFILER
+CSV_DEFINE_CATEGORY(ZenServerStats, true);
+
+CSV_DEFINE_STAT(ZenServerStats, ThroughputMbps);
+CSV_DEFINE_STAT(ZenServerStats, MaxReqThroughputMbps);
+CSV_DEFINE_STAT(ZenServerStats, MinReqThroughputMbps);
+CSV_DEFINE_STAT(ZenServerStats, RequestCountPerSec);
+#endif
+
+bool UDebugStorageServerConnection::ShowGraphs = true;
+
+void UDebugStorageServerConnection::StartDrawing()
+{
+	if (DrawHandle.IsValid())
+	{
+		return;
+	}
+	DrawHandle = UDebugDrawService::Register(TEXT("Game"),FDebugDrawDelegate::CreateUObject(this, &UDebugStorageServerConnection::Draw));
+}
+
+void UDebugStorageServerConnection::StopDrawing()
+{
+	if (!DrawHandle.IsValid())
+	{
+		return;
+	}
+
+	UDebugDrawService::Unregister(DrawHandle);
+	DrawHandle.Reset();
+}
+
+void UDebugStorageServerConnection::AddTimingInstance(double duration, uint64 bytes)
+{
+	if ((duration >= 0.0))
+	{
+		double tr = ((double)(bytes * 8) / duration) / 1000000.0; //Mbps
+
+		{
+			FScopeLock Lock(&StatsCS);
+			AccumulatedBytes += bytes;
+			RequestCount++;
+
+			if ((MinRequestThroughput == 0.0) || (tr < MinRequestThroughput))
+			{
+				MinRequestThroughput = tr;
+			}
+
+			if ((MaxRequestThroughput == 0.0) || (tr > MaxRequestThroughput))
+			{
+				MaxRequestThroughput = tr;
+			}
+		}
+	}
+}
+
+void UDebugStorageServerConnection::Draw(UCanvas* Canvas, APlayerController*)
+{
+	static constexpr double FrameSeconds = 1.0;
+	static constexpr float  ViewXRel = 0.2f;
+	static constexpr float  ViewYRel = 0.12f;
+	static constexpr float  ViewWidthRel = 0.4f;
+	static constexpr float  ViewHeightRel = 0.18f;
+	static constexpr double TextHeight = 16.0;
+	static constexpr int    OneMinuteSeconds = 60;
+	static constexpr double WidthSeconds = OneMinuteSeconds * 0.25;
+	static constexpr double MaxHeightScaleThroughput = 6000;
+	static constexpr double MaxHeightScaleRequest = 5000;
+	static constexpr int	LineThickness = 3;
+	static double			HeightScaleThroughput = MaxHeightScaleThroughput;
+	static double			HeightScaleRequest = MaxHeightScaleRequest;
+	
+	double StatsTimeNow = FPlatformTime::Seconds();
+	double Duration = StatsTimeNow - UpdateStatsTime;
+
+	static double MaxReqThroughput = 0.0;
+	static double MinReqThroughput = 0.0;
+	static uint32 ReqCount = 0;
+	static double Throughput = 0.0;
+
+	//Persistent debug message and CSV stats
+	if ((Duration > UpdateStatsTimer) && (GEngine))
+	{
+		UpdateStatsTime = StatsTimeNow;
+
+		{
+			FScopeLock Lock(&StatsCS);
+			MaxReqThroughput = MaxRequestThroughput;
+			MinReqThroughput = MinRequestThroughput;
+		
+			Throughput = ((double)(AccumulatedBytes * 8) / Duration) / 1000000.0; //Mbps
+			ReqCount = ceil((double)RequestCount / Duration);
+
+			AccumulatedBytes = 0;
+			RequestCount = 0;
+			MinRequestThroughput = 0.0;
+			MaxRequestThroughput = 0.0;
+		}
+
+		FString ZenConnectionDebugMsg;
+		ZenConnectionDebugMsg = FString::Printf(TEXT("ZenServer streaming from %s [%.2fMbps]"), *Owner->GetConnectionBackend()->GetHostName(), Throughput);
+		GEngine->AddOnScreenDebugMessage((uint64)this, 86400.0f, FColor::White, ZenConnectionDebugMsg, false);
+		
+		History.push_back({ StatsTimeNow, MaxReqThroughput, MinReqThroughput, Throughput, ReqCount });
+	}
+
+	while (!History.empty() && StatsTimeNow - History.front().Time > WidthSeconds)
+	{
+		History.erase(History.begin());
+	}
+
+#if CSV_PROFILER
+	//CSV stats need to be written per frame
+	CSV_CUSTOM_STAT_DEFINED(ThroughputMbps, Throughput, ECsvCustomStatOp::Set);
+	CSV_CUSTOM_STAT_DEFINED(MaxReqThroughputMbps, MaxReqThroughput, ECsvCustomStatOp::Set);
+	CSV_CUSTOM_STAT_DEFINED(MinReqThroughputMbps, MinReqThroughput, ECsvCustomStatOp::Set);
+	CSV_CUSTOM_STAT_DEFINED(RequestCountPerSec, (int32)ReqCount, ECsvCustomStatOp::Set);
+#endif
+
+	if (ShowGraphs)
+	{
+		int ViewX = (int)(ViewXRel * Canvas->ClipX);
+		int ViewY = (int)(ViewYRel * Canvas->ClipY);
+		int ViewWidth = (int)(ViewWidthRel * Canvas->ClipX);;
+		int ViewHeight = (int)(ViewHeightRel * Canvas->ClipY);;
+		double PixelsPerSecond = ViewWidth / WidthSeconds;
+
+		auto DrawLine =
+			[Canvas](double X0, double Y0, double X1, double Y1, const FLinearColor& Color, double Thickness)
+			{
+				FCanvasLineItem Line{ FVector2D{X0, Y0}, FVector2D{X1, Y1} };
+				Line.SetColor(Color);
+				Line.LineThickness = Thickness;
+				Canvas->DrawItem(Line);
+			};
+
+		auto DrawString =
+			[Canvas](const FString& Str, int X, int Y, bool bCentre = true)
+			{
+				FCanvasTextItem Text{ FVector2D(X, Y), FText::FromString(Str), GEngine->GetTinyFont(), FLinearColor::Yellow };
+				Text.EnableShadow(FLinearColor::Black);
+				Text.bCentreX = bCentre;
+				Text.bCentreY = bCentre;
+				Canvas->DrawItem(Text);
+			};
+
+		double MaxValueInHistory = 0.0;
+
+		if (History.size())
+		{
+			ViewY += TextHeight;
+			DrawString(FString::Printf(TEXT("Request Throughput MIN/MAX: [%.2f] / [%.2f] Mbps"), History[History.size()-1].MinRequestThroughput, History[History.size()-1].MaxRequestThroughput), ViewX, ViewY, false);
+			ViewY += TextHeight;
+		}
+
+		//FIRST GRAPH
+		MaxValueInHistory = 0.0;
+		double HeightScale = HeightScaleThroughput;
+		ViewY += TextHeight;
+		{ // draw graph edges + label
+			const FLinearColor Color = FLinearColor::White;
+			DrawLine(ViewX, ViewY + ViewHeight, ViewX + ViewWidth, ViewY + ViewHeight, Color, 1);
+			DrawLine(ViewX, ViewY, ViewX, ViewY + ViewHeight, Color, 1);
+			DrawLine(ViewX + ViewWidth, ViewY, ViewX + ViewWidth, ViewY + ViewHeight, Color, 1);
+			DrawString(TEXT("ZenServer Throughput Mbps"), ViewX, ViewY + ViewHeight + 10, false);
+		}
+
+		for (int I = History.size() - 1; I >= 0; --I)
+		{
+			const auto& Item = History[I];
+			const double       X = ViewX + ViewWidth - PixelsPerSecond * (StatsTimeNow - Item.Time);
+			const double       H = FMath::Min(ViewHeight, ViewHeight * (Item.Throughput / HeightScale));
+			const double       Y = ViewY + ViewHeight - H;
+			const FLinearColor Color = FLinearColor::Yellow;
+			
+			DrawLine(X, ViewY + ViewHeight - 1, X, Y, Color, LineThickness);
+			DrawString(FString::Printf(TEXT("%.2f"), Item.Throughput), X, Y - 11);
+
+			if (Item.Throughput > MaxValueInHistory)
+				MaxValueInHistory = Item.Throughput;
+		}
+		HeightScaleThroughput = FMath::Min(MaxHeightScaleThroughput, FMath::Max(MaxValueInHistory, 1.0));
+
+		//SECOND GRAPH
+		MaxValueInHistory = 0.0;
+		ViewY += ViewHeight + (TextHeight * 2) ;
+		HeightScale = HeightScaleRequest;
+
+		{ // draw graph edges + label
+			const FLinearColor Color = FLinearColor::White;
+			DrawLine(ViewX, ViewY + ViewHeight, ViewX + ViewWidth, ViewY + ViewHeight, Color, 1);
+			DrawLine(ViewX, ViewY, ViewX, ViewY + ViewHeight, Color, 1);
+			DrawLine(ViewX + ViewWidth, ViewY, ViewX + ViewWidth, ViewY + ViewHeight, Color, 1);
+			DrawString(TEXT("ZenServer Request/Sec Count"), ViewX, ViewY + ViewHeight + 10, false);
+		}
+
+		for (int I = History.size() - 1; I >= 0; --I)
+		{
+			const auto& Item = History[I];
+			const double       X = ViewX + ViewWidth - PixelsPerSecond * (StatsTimeNow - Item.Time);
+			const double       H = FMath::Min(ViewHeight, ViewHeight * (Item.RequestCount / HeightScale));
+			const double       Y = ViewY + ViewHeight - H;
+			const FLinearColor Color = FLinearColor::Gray;
+
+			DrawLine(X, ViewY + ViewHeight - 1, X, Y, Color, LineThickness);
+			DrawString(FString::Printf(TEXT("%.d"), Item.RequestCount), X, Y - 11);
+
+			if (Item.RequestCount > MaxValueInHistory)
+				MaxValueInHistory = Item.RequestCount;
+		}
+		HeightScaleRequest = FMath::Min(MaxHeightScaleRequest, FMath::Max(MaxValueInHistory, 1.0));
+	}
+}
+
+void UDebugStorageServerConnection::ShowGraph(FOutputDevice&)
+{
+	ShowGraphs = true;
+}
+
+void UDebugStorageServerConnection::HideGraph(FOutputDevice&)
+{
+	ShowGraphs = false;
+}
+
+static FAutoConsoleCommandWithOutputDevice
+	GShowDebugConnectionStatsCmd(
+		TEXT("r.ZenServerStatsShow"),
+		TEXT("Show ZenServer Stats Graph."),
+		FConsoleCommandWithOutputDeviceDelegate::CreateStatic(&UDebugStorageServerConnection::ShowGraph));
+
+static FAutoConsoleCommandWithOutputDevice
+	GHideDebugConnectionStatsCmd(
+		TEXT("r.ZenServerStatsHide"),
+		TEXT("Hide ZenServer Stats Graph."),
+		FConsoleCommandWithOutputDeviceDelegate::CreateStatic(&UDebugStorageServerConnection::HideGraph));
\ No newline at end of file
diff --git a/Engine/Source/Runtime/StorageServerClient/Private/DebugStorageServerConnection.h b/Engine/Source/Runtime/StorageServerClient/Private/DebugStorageServerConnection.h
new file mode 100644
--- /dev/null
+++ b/Engine/Source/Runtime/StorageServerClient/Private/DebugStorageServerConnection.h
@@ -0,0 +1,61 @@
+// Copyright Epic Games, Inc. All Rights Reserved.
+
+#pragma once
+
+#include "CoreMinimal.h"
+#include "Engine/Canvas.h"
+#include <vector>
+
+#include "DebugStorageServerConnection.generated.h"
+
+class FStorageServerConnection;
+
+UCLASS()
+class UDebugStorageServerConnection : public UObject
+{
+	GENERATED_BODY()
+
+public:
+	void StartDrawing();
+	void StopDrawing();
+
+	void AddTimingInstance(double duration, uint64 bytes);
+
+	void SetOwner(FStorageServerConnection* Connection)
+	{
+		Owner = Connection;
+	}
+
+	static void ShowGraph(FOutputDevice&);
+	static void HideGraph(FOutputDevice&);
+
+private:
+	void Draw(UCanvas* Canvas, class APlayerController* PC);
+
+	struct HistoryItem
+	{
+		double Time;
+		double MaxRequestThroughput;
+		double MinRequestThroughput;
+		double Throughput;
+		uint32 RequestCount;
+	};
+
+	std::vector<HistoryItem> History = {{0, 0, 0, 0, 0}};
+
+	FDelegateHandle DrawHandle;
+
+	static constexpr float UpdateStatsTimer = 1.0;
+	double UpdateStatsTime = 0.0;
+	uint64 AccumulatedBytes = 0;
+	uint32 RequestCount = 0;
+
+	double MinRequestThroughput = 0.0;
+	double MaxRequestThroughput = 0.0;
+
+	FCriticalSection StatsCS;
+	FStorageServerConnection* Owner;
+
+	static bool ShowGraphs;
+
+};
diff --git a/Engine/Source/Runtime/StorageServerClient/Private/StorageServerConnection.cpp b/Engine/Source/Runtime/StorageServerClient/Private/StorageServerConnection.cpp
--- a/Engine/Source/Runtime/StorageServerClient/Private/StorageServerConnection.cpp
+++ b/Engine/Source/Runtime/StorageServerClient/Private/StorageServerConnection.cpp
@@ -1,6 +1,7 @@
 // Copyright Epic Games, Inc. All Rights Reserved.
 
 #include "StorageServerConnection.h"
+#include "DebugStorageServerConnection.h"
 
 #include "IO/IoDispatcher.h"
 #include "IPAddress.h"
@@ -55,6 +56,7 @@ FStorageServerRequest::FStorageServerRequest(FAnsiStringView Verb, FAnsiStringVi
 
 IStorageConnectionSocket* FStorageServerRequest::Send(FStorageServerConnection& Owner, bool bLogOnError)
 {
+	StartTiming();
 	if (BodyBuffer.Num())
 	{
 		HeaderBuffer.Append("Content-Length: ").Appendf("%d\r\n", BodyBuffer.Num());
@@ -263,6 +265,7 @@ bool FStorageServerChunkBatchRequest::Issue(TFunctionRef<void(uint32 ChunkCount,
 		ChunkSizes.Emplace(ChunkSize);
 	}
 	OnResponse(ChunkCount, ChunkIndices.GetData(), ChunkSizes.GetData(), Response);
+	Owner.AddTimingInstance(GetDuration(), (double)Response.Tell());
 	return true;
 }
 
@@ -470,6 +473,11 @@ FStorageServerConnection::FStorageServerConnection()
 
 FStorageServerConnection::~FStorageServerConnection()
 {
+	if (StatsObject)
+	{
+		StatsObject->StopDrawing();
+		StatsObject->RemoveFromRoot();
+	}
 }
 
 bool FStorageServerConnection::Initialize(TArrayView<const FString> InHostAddresses, int32 InPort, const TCHAR* InProjectNameOverride, const TCHAR* InPlatformNameOverride)
@@ -507,7 +515,6 @@ bool FStorageServerConnection::Initialize(TArrayView<const FString> InHostAddres
 void FStorageServerConnection::PackageStoreRequest(TFunctionRef<void(FPackageStoreEntryResource&&)> Callback)
 {
 	TRACE_CPUPROFILER_EVENT_SCOPE(StorageServerPackageStoreRequest);
-	ShowDebugMessage();
 
 	TAnsiStringBuilder<256> ResourceBuilder;
 	ResourceBuilder.Append(OplogPath).Append("/entries?fieldfilter=packagestoreentry");
@@ -539,8 +546,6 @@ void FStorageServerConnection::PackageStoreRequest(TFunctionRef<void(FPackageSto
 
 void FStorageServerConnection::FileManifestRequest(TFunctionRef<void(FIoChunkId Id, FStringView Path)> Callback)
 {
-	ShowDebugMessage();
-
 	TAnsiStringBuilder<256> ResourceBuilder;
 	ResourceBuilder.Append(OplogPath).Append("/files?filter=client");
 	FStorageServerRequest Request("GET", *ResourceBuilder, Hostname, EStorageServerContentType::CbObject);
@@ -592,6 +597,8 @@ int64 FStorageServerConnection::ChunkSizeRequest(const FIoChunkId& ChunkId)
 	FStorageServerResponse Response(*this, *Socket);
 	if (Response.IsOk())
 	{
+		AddTimingInstance(Request.GetDuration(), (double)Response.TotalSize());
+
 		FCbObject ResponseObj = Response.GetResponseObject();
 
 		const int64 ChunkSize = ResponseObj["size"].AsInt64(0);
@@ -612,9 +619,6 @@ int64 FStorageServerConnection::ChunkSizeRequest(const FIoChunkId& ChunkId)
 bool FStorageServerConnection::ReadChunkRequest(const FIoChunkId& ChunkId, uint64 Offset, uint64 Size, TFunctionRef<void(FStorageServerResponse&)> OnResponse)
 {
 	TRACE_CPUPROFILER_EVENT_SCOPE(ZenHttpClient::ReadChunkRequest);
-
-	ShowDebugMessage();
-
 	TAnsiStringBuilder<256> ResourceBuilder;
 	ResourceBuilder.Append(OplogPath) << "/" << ChunkId;
 
@@ -656,6 +660,7 @@ bool FStorageServerConnection::ReadChunkRequest(const FIoChunkId& ChunkId, uint6
 	if (Response.IsOk())
 	{
 		OnResponse(Response);
+		AddTimingInstance(Request.GetDuration(), (double)Response.Tell());
 		return true;
 	}
 	else if (Response.GetErrorCode() == 404)
@@ -720,17 +725,19 @@ FStorageConnectionBackend* FStorageServerConnection::GetConnectionBackend() cons
 	return ConnectionBackend.Get();
 }
 
-void FStorageServerConnection::ShowDebugMessage()
+void FStorageServerConnection::AddTimingInstance(double duration, uint64 bytes)
 {
-	static bool bShowDebugMessage = true;
-
-	if ((bShowDebugMessage) && (GEngine))
+	if ((StatsObject == nullptr) && (UObjectInitialized()))
 	{
-		FString ZenConnectionDebugMsg;
-		ZenConnectionDebugMsg = FString::Printf(TEXT("ZenServer streaming from %s"), *ConnectionBackend->GetHostName());
-		GEngine->AddOnScreenDebugMessage(INDEX_NONE, 86400.0f, FColor::White, ZenConnectionDebugMsg, false);
-		bShowDebugMessage = false;
+		StatsObject = NewObject<UDebugStorageServerConnection>();
+		StatsObject->SetOwner(this);
+		StatsObject->AddToRoot();
+		StatsObject->StartDrawing();
 	}
+
+	if (StatsObject)
+		StatsObject->AddTimingInstance(duration, bytes);
 }
 
+
 #endif
diff --git a/Engine/Source/Runtime/StorageServerClient/Public/StorageServerConnection.h b/Engine/Source/Runtime/StorageServerClient/Public/StorageServerConnection.h
--- a/Engine/Source/Runtime/StorageServerClient/Public/StorageServerConnection.h
+++ b/Engine/Source/Runtime/StorageServerClient/Public/StorageServerConnection.h
@@ -9,6 +9,7 @@
 #include "Memory/MemoryFwd.h"
 #include "Templates/SharedPointer.h"
 #include "SocketTypes.h"
+#include "HAL/PlatformTime.h"
 
 #if !UE_BUILD_SHIPPING
 
@@ -95,6 +96,21 @@ protected:
 	EStorageServerContentType AcceptType;
 	TAnsiStringBuilder<512> HeaderBuffer;
 	TArray<uint8, TInlineAllocator<1024>> BodyBuffer;
+	double StartRequestWallTime;
+
+public:
+	void StartTiming()
+	{
+		StartRequestWallTime = FPlatformTime::Seconds();
+	}
+
+	double GetDuration() // seconds
+	{
+		check(StartRequestWallTime > 0.0);
+		double EndRequestWallTime = FPlatformTime::Seconds();
+		return EndRequestWallTime - StartRequestWallTime;
+	}
+
 };
 
 class FStorageServerResponse
@@ -218,6 +234,8 @@ protected:
 	TAnsiStringBuilder<1024> OplogPath;
 };
 
+class UDebugStorageServerConnection;
+
 class FStorageServerConnection
 {
 public:
@@ -238,13 +256,13 @@ private:
 	friend FStorageServerRequest;
 	friend FStorageServerResponse;
 	friend FStorageServerChunkBatchRequest;
+	friend UDebugStorageServerConnection;
 
 	bool CreateConnectionBackend(TArrayView<const FString> HostAddresses, int32 Port);
 	bool CreatePlatformBackend(const FString& HostAddresses, int32 Port);
 
 	FStorageConnectionBackend* GetConnectionBackend() const;
-
-	void ShowDebugMessage();
+	void AddTimingInstance(double duration, uint64 bytes);
 
 	ISocketSubsystem& SocketSubsystem;
 	TAnsiStringBuilder<1024> OplogPath;
@@ -252,6 +270,8 @@ private:
 	TAnsiStringBuilder<1024> Hostname;
 	TUniquePtr<FStorageConnectionBackend> ConnectionBackend;
 	FCriticalSection BackendCS;
+
+	UDebugStorageServerConnection* StatsObject = nullptr;
 };
 
 #endif
