ShowGraphs default value set to false

diff --git a/Engine/Source/Runtime/StorageServerClient/Private/DebugStorageServerConnection.cpp b/Engine/Source/Runtime/StorageServerClient/Private/DebugStorageServerConnection.cpp
--- a/Engine/Source/Runtime/StorageServerClient/Private/DebugStorageServerConnection.cpp
+++ b/Engine/Source/Runtime/StorageServerClient/Private/DebugStorageServerConnection.cpp
@@ -16,7 +16,7 @@ CSV_DEFINE_STAT(ZenServerStats, MinReqThroughputMbps);
 CSV_DEFINE_STAT(ZenServerStats, RequestCountPerSec);
 #endif
 
-bool UDebugStorageServerConnection::ShowGraphs = true;
+bool UDebugStorageServerConnection::ShowGraphs = false;
 
 void UDebugStorageServerConnection::StartDrawing()
 {
