Fixing Shipping build, disabling DebugStorageServerConnection since ZenServer is not available in that configuration

diff --git a/Engine/Source/Runtime/StorageServerClient/Private/DebugStorageServerConnection.cpp b/Engine/Source/Runtime/StorageServerClient/Private/DebugStorageServerConnection.cpp
--- a/Engine/Source/Runtime/StorageServerClient/Private/DebugStorageServerConnection.cpp
+++ b/Engine/Source/Runtime/StorageServerClient/Private/DebugStorageServerConnection.cpp
@@ -7,6 +7,8 @@
 #include "Engine/GameEngine.h"
 #include "ProfilingDebugging/CsvProfiler.h"
 
+#if !UE_BUILD_SHIPPING
+
 #if CSV_PROFILER
 CSV_DEFINE_CATEGORY(ZenServerStats, true);
 
@@ -240,4 +242,6 @@ static FAutoConsoleCommandWithOutputDevice
 	GHideDebugConnectionStatsCmd(
 		TEXT("r.ZenServerStatsHide"),
 		TEXT("Hide ZenServer Stats Graph."),
-		FConsoleCommandWithOutputDeviceDelegate::CreateStatic(&UDebugStorageServerConnection::HideGraph));
\ No newline at end of file
+		FConsoleCommandWithOutputDeviceDelegate::CreateStatic(&UDebugStorageServerConnection::HideGraph));
+
+#endif
\ No newline at end of file
