Cache file size from manifest in client StorageServer in order to remove a trip to server when asking for a file size.

diff --git a/Engine/Source/Runtime/StorageServerClient/Private/StorageServerConnection.cpp b/Engine/Source/Runtime/StorageServerClient/Private/StorageServerConnection.cpp
--- a/Engine/Source/Runtime/StorageServerClient/Private/StorageServerConnection.cpp
+++ b/Engine/Source/Runtime/StorageServerClient/Private/StorageServerConnection.cpp
@@ -307,6 +307,8 @@ void FStorageServerResponse::Serialize(void* V, int64 Length)
 		Position += BytesRead;
 	}
 
+	check(ContentType == EStorageServerContentType::Binary);
+
 	TRACE_COUNTER_ADD(ZenHttpClientSerializedBytes, Length);
 	
 	if (Position == ContentLength)
@@ -544,10 +546,10 @@ void FStorageServerConnection::PackageStoreRequest(TFunctionRef<void(FPackageSto
 	}
 }
 
-void FStorageServerConnection::FileManifestRequest(TFunctionRef<void(FIoChunkId Id, FStringView Path)> Callback)
+void FStorageServerConnection::FileManifestRequest(TFunctionRef<void(FIoChunkId Id, FStringView Path, int64 RawSize)> Callback)
 {
 	TAnsiStringBuilder<256> ResourceBuilder;
-	ResourceBuilder.Append(OplogPath).Append("/files?filter=client");
+	ResourceBuilder.Append(OplogPath).Append("/files?fieldnames=id,clientpath,rawsize");
 	FStorageServerRequest Request("GET", *ResourceBuilder, Hostname, EStorageServerContentType::CbObject);
 	IStorageConnectionSocket* Socket = Request.Send(*this);
 	if (!Socket)
@@ -565,6 +567,7 @@ void FStorageServerConnection::FileManifestRequest(TFunctionRef<void(FIoChunkId
 		{
 			FCbObject Entry = FileArrayEntry.AsObject();
 			FCbObjectId Id = Entry["id"].AsObjectId();
+			int64 ResponseRawSize = Entry["rawsize"].AsInt64();
 
 			TStringBuilder<128> WidePath;
 			WidePath.Append(FUTF8ToTCHAR(Entry["clientpath"].AsString()));
@@ -572,7 +575,8 @@ void FStorageServerConnection::FileManifestRequest(TFunctionRef<void(FIoChunkId
 			FIoChunkId ChunkId;
 			ChunkId.Set(Id.GetView());
 
-			Callback(ChunkId, WidePath);
+
+			Callback(ChunkId, WidePath, ResponseRawSize);
 		}
 	}
 	else
diff --git a/Engine/Source/Runtime/StorageServerClient/Private/StorageServerPlatformFile.cpp b/Engine/Source/Runtime/StorageServerClient/Private/StorageServerPlatformFile.cpp
--- a/Engine/Source/Runtime/StorageServerClient/Private/StorageServerPlatformFile.cpp
+++ b/Engine/Source/Runtime/StorageServerClient/Private/StorageServerPlatformFile.cpp
@@ -62,7 +62,7 @@ FStorageServerFileSystemTOC::FDirectory* FStorageServerFileSystemTOC::AddDirecto
 	return Directory;
 }
 
-void FStorageServerFileSystemTOC::AddFile(const FIoChunkId& FileChunkId, FStringView PathView)
+void FStorageServerFileSystemTOC::AddFile(const FIoChunkId& FileChunkId, FStringView PathView, int64 RawSize)
 {
 	FWriteScopeLock _(TocLock);
 
@@ -71,6 +71,7 @@ void FStorageServerFileSystemTOC::AddFile(const FIoChunkId& FileChunkId, FString
 	FFile& NewFile = Files.AddDefaulted_GetRef();
 	NewFile.FileChunkId = FileChunkId;
 	NewFile.FilePath = PathView;
+	NewFile.RawSize = RawSize;
 	
 	FilePathToIndexMap.Add(NewFile.FilePath, FileIndex);
 	
@@ -105,7 +106,30 @@ const FIoChunkId* FStorageServerFileSystemTOC::GetFileChunkId(const FString& Pat
 	return nullptr;
 }
 
-bool FStorageServerFileSystemTOC::IterateDirectory(const FString& Path, TFunctionRef<bool(const FIoChunkId&, const TCHAR*)> Callback)
+int64 FStorageServerFileSystemTOC::GetFileSize(const FString& Path)
+{
+	FReadScopeLock _(TocLock);
+	if (const int32* FileIndex = FilePathToIndexMap.Find(Path))
+	{
+		return Files[*FileIndex].RawSize;
+	}
+	return STORAGE_SERVER_FILE_UNKOWN_SIZE;
+}
+
+bool FStorageServerFileSystemTOC::GetFileData(const FString& Path, FIoChunkId& OutChunkId, int64& OutRawSize)
+{
+	FReadScopeLock _(TocLock);
+	if (const int32* FileIndex = FilePathToIndexMap.Find(Path))
+	{
+		const FFile& File = Files[*FileIndex];
+		OutChunkId = File.FileChunkId;
+		OutRawSize = File.RawSize;
+		return true;
+	}
+	return false;
+}
+
+bool FStorageServerFileSystemTOC::IterateDirectory(const FString& Path, TFunctionRef<bool(const FIoChunkId&, const TCHAR*, int64 RawSize)> Callback)
 {
 	UE_LOG(LogStorageServerPlatformFile, Verbose, TEXT("IterateDirectory '%s'"), *Path);
 
@@ -119,14 +143,14 @@ bool FStorageServerFileSystemTOC::IterateDirectory(const FString& Path, TFunctio
 	for (int32 FileIndex : Directory->Files)
 	{
 		const FFile& File = Files[FileIndex];
-		if (!Callback(File.FileChunkId, *File.FilePath))
+		if (!Callback(File.FileChunkId, *File.FilePath, File.RawSize))
 		{
 			return false;
 		}
 	}
 	for (const FString& ChildDirectoryPath : Directory->Directories)
 	{
-		if (!Callback(FIoChunkId(), *ChildDirectoryPath))
+		if (!Callback(FIoChunkId(), *ChildDirectoryPath, 0))
 		{
 			return false;
 		}
@@ -151,10 +175,11 @@ class FStorageServerFileHandle
 	uint8 Buffer[BufferSize];
 
 public:
-	FStorageServerFileHandle(FStorageServerPlatformFile& InOwner, FIoChunkId InFileChunkId, const TCHAR* InFilename)
+	FStorageServerFileHandle(FStorageServerPlatformFile& InOwner, FIoChunkId InFileChunkId, int64 InFileSize, const TCHAR* InFilename)
 		: Owner(InOwner)
 		, FileChunkId(InFileChunkId)
 		, Filename(InFilename)
+		, FileSize(InFileSize)
 	{
 	}
 
@@ -519,14 +544,13 @@ int64 FStorageServerPlatformFile::FileSize(const TCHAR* Filename)
 	TStringBuilder<1024> StorageServerFilename;
 	if (MakeStorageServerPath(Filename, StorageServerFilename))
 	{
-		if (const FIoChunkId* FileChunkId = ServerToc.GetFileChunkId(*StorageServerFilename))
+		int64 FileSize = ServerToc.GetFileSize(*StorageServerFilename);
+		if (FileSize > STORAGE_SERVER_FILE_UNKOWN_SIZE)
 		{
-			const FFileStatData FileStatData = SendGetStatDataMessage(*FileChunkId);
-			check(FileStatData.bIsValid);
-			return FileStatData.FileSize;
+			return FileSize;
 		}
 	}
-	return IsNonServerFilenameAllowed(Filename) ? LowerLevel->FileSize(Filename) : -1;
+	return IsNonServerFilenameAllowed(Filename) ? LowerLevel->FileSize(Filename) : STORAGE_SERVER_FILE_UNKOWN_SIZE;
 }
 
 bool FStorageServerPlatformFile::IsReadOnly(const TCHAR* Filename)
@@ -544,9 +568,16 @@ FFileStatData FStorageServerPlatformFile::GetStatData(const TCHAR* FilenameOrDir
 	TStringBuilder<1024> StorageServerFilenameOrDirectory;
 	if (MakeStorageServerPath(FilenameOrDirectory, StorageServerFilenameOrDirectory))
 	{
-		if (const FIoChunkId* FileChunkId = ServerToc.GetFileChunkId(*StorageServerFilenameOrDirectory))
+		int64 FileSize = ServerToc.GetFileSize(*StorageServerFilenameOrDirectory);
+		if (FileSize > STORAGE_SERVER_FILE_UNKOWN_SIZE)
 		{
-			return SendGetStatDataMessage(*FileChunkId);
+			return FFileStatData(
+				FDateTime::Now(),
+				FDateTime::Now(),
+				FDateTime::Now(),
+				FileSize,
+				false,
+				true);
 		}
 		else if (ServerToc.DirectoryExists(*StorageServerFilenameOrDirectory))
 		{
@@ -567,9 +598,9 @@ FFileStatData FStorageServerPlatformFile::GetStatData(const TCHAR* FilenameOrDir
 	return FileStatData;
 }
 
-IFileHandle* FStorageServerPlatformFile::InternalOpenFile(const FIoChunkId& FileChunkId, const TCHAR* LocalFilename)
+IFileHandle* FStorageServerPlatformFile::InternalOpenFile(const FIoChunkId& FileChunkId, int64 RawSize, const TCHAR* LocalFilename)
 {
-	return new FStorageServerFileHandle(*this, FileChunkId, LocalFilename);
+	return new FStorageServerFileHandle(*this, FileChunkId, RawSize, LocalFilename);
 }
 
 IFileHandle* FStorageServerPlatformFile::OpenRead(const TCHAR* Filename, bool bAllowWrite)
@@ -577,9 +608,11 @@ IFileHandle* FStorageServerPlatformFile::OpenRead(const TCHAR* Filename, bool bA
 	TStringBuilder<1024> StorageServerFilename;
 	if (MakeStorageServerPath(Filename, StorageServerFilename))
 	{
-		if (const FIoChunkId* FileChunkId = ServerToc.GetFileChunkId(*StorageServerFilename))
+		FIoChunkId FileChunkId;
+		int64 RawSize = STORAGE_SERVER_FILE_UNKOWN_SIZE;
+		if (ServerToc.GetFileData(*StorageServerFilename, FileChunkId, RawSize))
 		{
-			return InternalOpenFile(*FileChunkId, Filename);
+			return InternalOpenFile(FileChunkId, RawSize, Filename);
 		}
 	}
 	return IsNonServerFilenameAllowed(Filename) ? LowerLevel->OpenRead(Filename, bAllowWrite) : nullptr;
@@ -591,7 +624,7 @@ bool FStorageServerPlatformFile::IterateDirectory(const TCHAR* Directory, IPlatf
 	bool bResult = false;
 	if (MakeStorageServerPath(Directory, StorageServerDirectory) && ServerToc.DirectoryExists(*StorageServerDirectory))
 	{
-		bResult |= ServerToc.IterateDirectory(*StorageServerDirectory, [this, &Visitor](const FIoChunkId& FileChunkId, const TCHAR* FilenameOrDirectory)
+		bResult |= ServerToc.IterateDirectory(*StorageServerDirectory, [this, &Visitor](const FIoChunkId& FileChunkId, const TCHAR* FilenameOrDirectory, int64 RawSize)
 		{
 			TStringBuilder<1024> LocalPath;
 			bool bConverted = MakeLocalPath(FilenameOrDirectory, LocalPath);
@@ -613,7 +646,7 @@ bool FStorageServerPlatformFile::IterateDirectoryStat(const TCHAR* Directory, FD
 	bool bResult = false;
 	if (MakeStorageServerPath(Directory, StorageServerDirectory) && ServerToc.DirectoryExists(*StorageServerDirectory))
 	{
-		bResult |= ServerToc.IterateDirectory(*StorageServerDirectory, [this, &Visitor](const FIoChunkId& FileChunkId, const TCHAR* ServerFilenameOrDirectory)
+		bResult |= ServerToc.IterateDirectory(*StorageServerDirectory, [this, &Visitor](const FIoChunkId& FileChunkId, const TCHAR* ServerFilenameOrDirectory, int64 RawSize)
 		{
 			TStringBuilder<1024> LocalPath;
 			bool bConverted = MakeLocalPath(ServerFilenameOrDirectory, LocalPath);
@@ -621,7 +654,13 @@ bool FStorageServerPlatformFile::IterateDirectoryStat(const TCHAR* Directory, FD
 			FFileStatData FileStatData;
 			if (FileChunkId.IsValid())
 			{
-				FileStatData = SendGetStatDataMessage(FileChunkId);
+				FileStatData = FFileStatData(
+					FDateTime::Now(),
+					FDateTime::Now(),
+					FDateTime::Now(),
+					RawSize,
+					false,
+					true);
 				check(FileStatData.bIsValid);
 			}
 			else
@@ -690,7 +729,9 @@ bool FStorageServerPlatformFile::MoveFile(const TCHAR* To, const TCHAR* From)
 	TStringBuilder<1024> StorageServerFrom;
 	if (MakeStorageServerPath(From, StorageServerFrom))
 	{
-		if (const FIoChunkId* FromFileChunkId = ServerToc.GetFileChunkId(*StorageServerFrom))
+		FIoChunkId FromFileChunkId;
+		int64 FromFileRawSize = STORAGE_SERVER_FILE_UNKOWN_SIZE;
+		if (ServerToc.GetFileData(*StorageServerFrom, FromFileChunkId, FromFileRawSize))
 		{
 			TUniquePtr<IFileHandle> ToFile(LowerLevel->OpenWrite(To, false, false));
 			if (!ToFile)
@@ -698,7 +739,7 @@ bool FStorageServerPlatformFile::MoveFile(const TCHAR* To, const TCHAR* From)
 				return false;
 			}
 
-			TUniquePtr<IFileHandle> FromFile(InternalOpenFile(*FromFileChunkId, *StorageServerFrom));
+			TUniquePtr<IFileHandle> FromFile(InternalOpenFile(FromFileChunkId, FromFileRawSize, *StorageServerFrom));
 			if (!FromFile)
 			{
 				return false;
@@ -897,9 +938,9 @@ bool FStorageServerPlatformFile::SendGetFileListMessage()
 {
 	TRACE_CPUPROFILER_EVENT_SCOPE(StorageServerPlatformFileGetFileList);
 	
-	Connection->FileManifestRequest([&](FIoChunkId Id, FStringView Path)
+	Connection->FileManifestRequest([&](FIoChunkId Id, FStringView Path, int64 RawSize)
 	{
-		ServerToc.AddFile(Id, Path);
+		ServerToc.AddFile(Id, Path, RawSize);
 	});
 
 	return true;
@@ -983,7 +1024,7 @@ void FStorageServerPlatformFile::OnCookOnTheFlyMessage(const UE::Cook::FCookOnTh
 			for (int32 Idx = 0, Num = Filenames.Num(); Idx < Num; ++Idx)
 			{
 				UE_LOG(LogCookOnTheFly, Verbose, TEXT("Adding file '%s'"), *Filenames[Idx]);
-				ServerToc.AddFile(ChunkIds[Idx], Filenames[Idx]);
+				ServerToc.AddFile(ChunkIds[Idx], Filenames[Idx], STORAGE_SERVER_FILE_UNKOWN_SIZE);
 			}
 
 			break;
diff --git a/Engine/Source/Runtime/StorageServerClient/Private/StorageServerPlatformFile.h b/Engine/Source/Runtime/StorageServerClient/Private/StorageServerPlatformFile.h
--- a/Engine/Source/Runtime/StorageServerClient/Private/StorageServerPlatformFile.h
+++ b/Engine/Source/Runtime/StorageServerClient/Private/StorageServerPlatformFile.h
@@ -21,15 +21,19 @@ namespace UE::Cook
 }
 #endif
 
+#define STORAGE_SERVER_FILE_UNKOWN_SIZE (-1) 
+
 class FStorageServerFileSystemTOC
 {
 public:
 	~FStorageServerFileSystemTOC();
-	void AddFile(const FIoChunkId& FileChunkId, FStringView Path);
+	void AddFile(const FIoChunkId& FileChunkId, FStringView Path, int64 RawSize);
 	bool FileExists(const FString& Path);
 	bool DirectoryExists(const FString& Path);
 	const FIoChunkId* GetFileChunkId(const FString& Path);
-	bool IterateDirectory(const FString& Path, TFunctionRef<bool(const FIoChunkId&, const TCHAR*)> Callback);
+	int64 GetFileSize(const FString& Path);
+	bool GetFileData(const FString& Path, FIoChunkId& OutChunkId, int64& OutRawSize);
+	bool IterateDirectory(const FString& Path, TFunctionRef<bool(const FIoChunkId&, const TCHAR*, int64)> Callback);
 
 private:
 	struct FDirectory
@@ -42,6 +46,7 @@ private:
 	{
 		FIoChunkId FileChunkId;
 		FString FilePath;
+		int64 RawSize;
 	};
 
 	FDirectory* AddDirectoriesRecursive(const FString& DirectoryPath);
@@ -106,7 +111,7 @@ private:
 	bool IsNonServerFilenameAllowed(FStringView InFilename);
 	bool MakeStorageServerPath(const TCHAR* LocalFilenameOrDirectory, FStringBuilderBase& OutPath) const;
 	bool MakeLocalPath(const TCHAR* ServerFilenameOrDirectory, FStringBuilderBase& OutPath) const;
-	IFileHandle* InternalOpenFile(const FIoChunkId& FileChunkId, const TCHAR* LocalFilename);
+	IFileHandle* InternalOpenFile(const FIoChunkId& FileChunkId, int64 RawSize, const TCHAR* LocalFilename);
 	bool SendGetFileListMessage();
 	FFileStatData SendGetStatDataMessage(const FIoChunkId& FileChunkId);
 	int64 SendReadMessage(uint8* Destination, const FIoChunkId& FileChunkId, int64 Offset, int64 BytesToRead);
diff --git a/Engine/Source/Runtime/StorageServerClient/Public/StorageServerConnection.h b/Engine/Source/Runtime/StorageServerClient/Public/StorageServerConnection.h
--- a/Engine/Source/Runtime/StorageServerClient/Public/StorageServerConnection.h
+++ b/Engine/Source/Runtime/StorageServerClient/Public/StorageServerConnection.h
@@ -245,7 +245,7 @@ public:
 	STORAGESERVERCLIENT_API bool Initialize(TArrayView<const FString> HostAddresses, int32 Port, const TCHAR* ProjectNameOverride = nullptr, const TCHAR* PlatformNameOverride = nullptr);
 
 	STORAGESERVERCLIENT_API void PackageStoreRequest(TFunctionRef<void(FPackageStoreEntryResource&&)> Callback);
-	STORAGESERVERCLIENT_API void FileManifestRequest(TFunctionRef<void(FIoChunkId Id, FStringView Path)> Callback);
+	STORAGESERVERCLIENT_API void FileManifestRequest(TFunctionRef<void(FIoChunkId Id, FStringView Path, int64 RawSize)> Callback);
 	STORAGESERVERCLIENT_API int64 ChunkSizeRequest(const FIoChunkId& ChunkId);
 	STORAGESERVERCLIENT_API bool ReadChunkRequest(const FIoChunkId& ChunkId, uint64 Offset, uint64 Size, TFunctionRef<void(FStorageServerResponse&)> OnResponse);
 	STORAGESERVERCLIENT_API FStorageServerChunkBatchRequest NewChunkBatchRequest();
