Remove temporary check from FStorageServerResponse::Serialize

diff --git a/Engine/Source/Runtime/StorageServerClient/Private/StorageServerConnection.cpp b/Engine/Source/Runtime/StorageServerClient/Private/StorageServerConnection.cpp
--- a/Engine/Source/Runtime/StorageServerClient/Private/StorageServerConnection.cpp
+++ b/Engine/Source/Runtime/StorageServerClient/Private/StorageServerConnection.cpp
@@ -307,8 +307,6 @@ void FStorageServerResponse::Serialize(void* V, int64 Length)
 		Position += BytesRead;
 	}
 
-	check(ContentType == EStorageServerContentType::Binary);
-
 	TRACE_COUNTER_ADD(ZenHttpClientSerializedBytes, Length);
 	
 	if (Position == ContentLength)
