Temporarily disable the caching of file sizes until an issue with CookOnTheFly mode is addressed.

diff --git a/Engine/Source/Runtime/StorageServerClient/Private/StorageServerConnection.cpp b/Engine/Source/Runtime/StorageServerClient/Private/StorageServerConnection.cpp
--- a/Engine/Source/Runtime/StorageServerClient/Private/StorageServerConnection.cpp
+++ b/Engine/Source/Runtime/StorageServerClient/Private/StorageServerConnection.cpp
@@ -547,7 +547,7 @@ void FStorageServerConnection::PackageStoreRequest(TFunctionRef<void(FPackageSto
 void FStorageServerConnection::FileManifestRequest(TFunctionRef<void(FIoChunkId Id, FStringView Path, int64 RawSize)> Callback)
 {
 	TAnsiStringBuilder<256> ResourceBuilder;
-	ResourceBuilder.Append(OplogPath).Append("/files?fieldnames=id,clientpath,rawsize");
+	ResourceBuilder.Append(OplogPath).Append("/files?filter=client");
 	FStorageServerRequest Request("GET", *ResourceBuilder, Hostname, EStorageServerContentType::CbObject);
 	IStorageConnectionSocket* Socket = Request.Send(*this);
 	if (!Socket)
@@ -565,7 +565,7 @@ void FStorageServerConnection::FileManifestRequest(TFunctionRef<void(FIoChunkId
 		{
 			FCbObject Entry = FileArrayEntry.AsObject();
 			FCbObjectId Id = Entry["id"].AsObjectId();
-			int64 ResponseRawSize = Entry["rawsize"].AsInt64();
+			int64 ResponseRawSize = Entry["rawsize"].AsInt64(-1);
 
 			TStringBuilder<128> WidePath;
 			WidePath.Append(FUTF8ToTCHAR(Entry["clientpath"].AsString()));
