Add support for specifying BaseURI when connecting to zenserver to read cooked/paked data

diff --git a/Engine/Source/Runtime/StorageServerClient/Private/StoragePlatformConnectionBackend.cpp b/Engine/Source/Runtime/StorageServerClient/Private/StoragePlatformConnectionBackend.cpp
--- a/Engine/Source/Runtime/StorageServerClient/Private/StoragePlatformConnectionBackend.cpp
+++ b/Engine/Source/Runtime/StorageServerClient/Private/StoragePlatformConnectionBackend.cpp
@@ -223,7 +223,7 @@ bool FStorageServerPlatformConnectionBackend::InitializeInternal(TArrayView<cons
 int32 FStorageServerPlatformConnectionBackend::HandshakeRequest(TArrayView<const FString> HostAddresses)
 {
 	TAnsiStringBuilder<256> ResourceBuilder;
-	ResourceBuilder.Append(OplogPath);
+	ResourceBuilder.Append(BaseURI);
 
 	for (const FString& Addr : HostAddresses)
 	{
diff --git a/Engine/Source/Runtime/StorageServerClient/Private/StorageServerConnection.cpp b/Engine/Source/Runtime/StorageServerClient/Private/StorageServerConnection.cpp
--- a/Engine/Source/Runtime/StorageServerClient/Private/StorageServerConnection.cpp
+++ b/Engine/Source/Runtime/StorageServerClient/Private/StorageServerConnection.cpp
@@ -434,38 +434,12 @@ FStorageConnectionBackend::FStorageConnectionBackend(FStorageServerConnection& I
 {
 }
 
-bool FStorageConnectionBackend::Initialize(TArrayView<const FString> InHostAddresses, int32 InPort, const TCHAR* InProjectNameOverride, const TCHAR* InPlatformNameOverride)
+bool FStorageConnectionBackend::Initialize(TArrayView<const FString> InHostAddresses, int32 InPort, const FAnsiStringView& InBaseURI)
 {
-	InitOplog(InProjectNameOverride, InPlatformNameOverride);
-
+	BaseURI = InBaseURI;
 	return InitializeInternal(InHostAddresses, InPort);
 }
 
-void FStorageConnectionBackend::InitOplog(const TCHAR* InProjectNameOverride, const TCHAR* InPlatformNameOverride)
-{
-	OplogPath.Append("/prj/");
-	if (InProjectNameOverride)
-	{
-		OplogPath.Append(TCHAR_TO_ANSI(InProjectNameOverride));
-	}
-	else
-	{
-		OplogPath.Append(TCHAR_TO_ANSI(*FApp::GetZenStoreProjectId()));
-	}
-	OplogPath.Append("/oplog/");
-	if (InPlatformNameOverride)
-	{
-		OplogPath.Append(TCHAR_TO_ANSI(InPlatformNameOverride));
-	}
-	else
-	{
-		TArray<FString> TargetPlatformNames;
-		FPlatformMisc::GetValidTargetPlatforms(TargetPlatformNames);
-		check(TargetPlatformNames.Num() > 0);
-		OplogPath.Append(TCHAR_TO_ANSI(*TargetPlatformNames[0]));
-	}
-}
-
 FStorageServerConnection::FStorageServerConnection()
 	: SocketSubsystem(*ISocketSubsystem::Get())
 {
@@ -480,36 +454,14 @@ FStorageServerConnection::~FStorageServerConnection()
 	}
 }
 
-bool FStorageServerConnection::Initialize(TArrayView<const FString> InHostAddresses, int32 InPort, const TCHAR* InProjectNameOverride, const TCHAR* InPlatformNameOverride)
+bool FStorageServerConnection::Initialize(TArrayView<const FString> InHostAddresses, int32 InPort, const FAnsiStringView& InBaseURI)
 {
-	OplogPath.Append("/prj/");
-	if (InProjectNameOverride)
-	{
-		OplogPath.Append(TCHAR_TO_ANSI(InProjectNameOverride));
-	}
-	else
-	{
-		OplogPath.Append(TCHAR_TO_ANSI(*FApp::GetZenStoreProjectId()));
-	}
-	OplogPath.Append("/oplog/");
-	if (InPlatformNameOverride)
-	{
-		OplogPath.Append(TCHAR_TO_ANSI(InPlatformNameOverride));
-	}
-	else
-	{
-		TArray<FString> TargetPlatformNames;
-		FPlatformMisc::GetValidTargetPlatforms(TargetPlatformNames);
-		check(TargetPlatformNames.Num() > 0);
-		OplogPath.Append(TCHAR_TO_ANSI(*TargetPlatformNames[0]));
-	}
-
+	BaseURI = InBaseURI;
 	if (!CreateConnectionBackend(InHostAddresses, InPort))
 	{
 		return false;
 	}
-	
-	return ConnectionBackend->Initialize(InHostAddresses, InPort, InProjectNameOverride, InPlatformNameOverride);
+	return ConnectionBackend->Initialize(InHostAddresses, InPort, BaseURI);
 }
 
 void FStorageServerConnection::PackageStoreRequest(TFunctionRef<void(FPackageStoreEntryResource&&)> Callback)
@@ -517,7 +469,7 @@ void FStorageServerConnection::PackageStoreRequest(TFunctionRef<void(FPackageSto
 	TRACE_CPUPROFILER_EVENT_SCOPE(StorageServerPackageStoreRequest);
 
 	TAnsiStringBuilder<256> ResourceBuilder;
-	ResourceBuilder.Append(OplogPath).Append("/entries?fieldfilter=packagestoreentry");
+	ResourceBuilder.Append(BaseURI).Append("/entries?fieldfilter=packagestoreentry");
 	FStorageServerRequest Request("GET", *ResourceBuilder, Hostname, EStorageServerContentType::CompressedBinary);
 	IStorageConnectionSocket* Socket = Request.Send(*this);
 	if (!Socket)
@@ -547,7 +499,7 @@ void FStorageServerConnection::PackageStoreRequest(TFunctionRef<void(FPackageSto
 void FStorageServerConnection::FileManifestRequest(TFunctionRef<void(FIoChunkId Id, FStringView Path, int64 RawSize)> Callback)
 {
 	TAnsiStringBuilder<256> ResourceBuilder;
-	ResourceBuilder.Append(OplogPath).Append("/files?fieldnames=id,clientpath,rawsize");
+	ResourceBuilder.Append(BaseURI).Append("/files?fieldnames=id,clientpath,rawsize");
 	FStorageServerRequest Request("GET", *ResourceBuilder, Hostname, EStorageServerContentType::CbObject);
 	IStorageConnectionSocket* Socket = Request.Send(*this);
 	if (!Socket)
@@ -586,7 +538,7 @@ void FStorageServerConnection::FileManifestRequest(TFunctionRef<void(FIoChunkId
 int64 FStorageServerConnection::ChunkSizeRequest(const FIoChunkId& ChunkId)
 {
 	TAnsiStringBuilder<256> ResourceBuilder;
-	ResourceBuilder.Append(OplogPath);
+	ResourceBuilder.Append(BaseURI);
 	ResourceBuilder << "/" << ChunkId << "/info";
 
 	FStorageServerRequest Request("GET", *ResourceBuilder, Hostname, EStorageServerContentType::CbObject);
@@ -622,7 +574,7 @@ bool FStorageServerConnection::ReadChunkRequest(const FIoChunkId& ChunkId, uint6
 {
 	TRACE_CPUPROFILER_EVENT_SCOPE(ZenHttpClient::ReadChunkRequest);
 	TAnsiStringBuilder<256> ResourceBuilder;
-	ResourceBuilder.Append(OplogPath) << "/" << ChunkId;
+	ResourceBuilder.Append(BaseURI) << "/" << ChunkId;
 
 	bool HaveQuery = false;
 
@@ -679,7 +631,7 @@ bool FStorageServerConnection::ReadChunkRequest(const FIoChunkId& ChunkId, uint6
 FStorageServerChunkBatchRequest FStorageServerConnection::NewChunkBatchRequest()
 {
 	TAnsiStringBuilder<256> ResourceBuilder;
-	ResourceBuilder.Append(OplogPath).Append("/batch");
+	ResourceBuilder.Append(BaseURI).Append("/batch");
 	return FStorageServerChunkBatchRequest(*this, *ResourceBuilder, Hostname);
 }
 
diff --git a/Engine/Source/Runtime/StorageServerClient/Private/StorageServerPlatformFile.cpp b/Engine/Source/Runtime/StorageServerClient/Private/StorageServerPlatformFile.cpp
--- a/Engine/Source/Runtime/StorageServerClient/Private/StorageServerPlatformFile.cpp
+++ b/Engine/Source/Runtime/StorageServerClient/Private/StorageServerPlatformFile.cpp
@@ -323,6 +323,41 @@ TUniquePtr<FArchive> FStorageServerPlatformFile::TryFindProjectStoreMarkerFile(I
 	return nullptr;
 }
 
+FAnsiString FStorageServerPlatformFile::MakeBaseURI()
+{
+	TAnsiStringBuilder<256> BaseURIBuilder;
+	if (!BaseURI.IsEmpty())
+	{
+		BaseURIBuilder.Append(BaseURI);
+	}
+	else
+	{
+		BaseURIBuilder.Append("/prj/");
+		if (ServerProject.IsEmpty())
+		{
+			BaseURIBuilder.Append(TCHAR_TO_ANSI(*FApp::GetZenStoreProjectId()));
+		}
+		else
+		{
+			BaseURIBuilder.Append(ServerProject);
+		}
+		BaseURIBuilder.Append("/oplog/");
+		if (ServerPlatform.IsEmpty())
+		{
+			TArray<FString> TargetPlatformNames;
+			FPlatformMisc::GetValidTargetPlatforms(TargetPlatformNames);
+			check(TargetPlatformNames.Num() > 0);
+			BaseURIBuilder.Append(TCHAR_TO_ANSI(*TargetPlatformNames[0]));
+		}
+		else
+		{
+			BaseURIBuilder.Append(ServerPlatform);
+		}
+	}
+	return BaseURIBuilder.ToString();
+}
+
+
 bool FStorageServerPlatformFile::ShouldBeUsed(IPlatformFile* Inner, const TCHAR* CmdLine) const
 {
 #if WITH_COTF
@@ -425,6 +460,7 @@ bool FStorageServerPlatformFile::Initialize(IPlatformFile* Inner, const TCHAR* C
 					const TSharedPtr<FJsonObject>& ZenServerObject = *ZenServerObjectPtr;
 					ServerProject = ZenServerObject->GetStringField(TEXT("projectid"));
 					ServerPlatform = ZenServerObject->GetStringField(TEXT("oplogid"));
+					BaseURI = ZenServerObject->GetStringField(TEXT("baseuri"));
 					UE_LOG(LogStorageServerPlatformFile, Display, TEXT("Using settings from ue.projectstore: ServerProject='%s' and ServerPlatform='%s'"), *ServerProject, *ServerPlatform);
 				}
 			}
@@ -438,6 +474,10 @@ bool FStorageServerPlatformFile::Initialize(IPlatformFile* Inner, const TCHAR* C
 		{
 			UE_LOG(LogStorageServerPlatformFile, Display, TEXT("Using settings from command line: -ZenStorePlatform='%s'"), *ServerPlatform);
 		}
+		if (FParse::Value(CmdLine, TEXT("-ZenStoreBaseURI="), BaseURI))
+		{
+			UE_LOG(LogStorageServerPlatformFile, Display, TEXT("Using settings from command line: -ZenStoreBaseURI='%s'"), *BaseURI);
+		}
 		return true;
 	}
 	return false;
@@ -456,9 +496,7 @@ void FStorageServerPlatformFile::InitializeAfterProjectFilePath()
 	}
 #endif
 	Connection.Reset(new FStorageServerConnection());
-	const TCHAR* ProjectOverride = ServerProject.IsEmpty() ? nullptr : *ServerProject;
-	const TCHAR* PlatformOverride = ServerPlatform.IsEmpty() ? nullptr : *ServerPlatform;
-	if (Connection->Initialize(HostAddrs, HostPort, ProjectOverride, PlatformOverride))
+	if (Connection->Initialize(HostAddrs, HostPort, MakeBaseURI()))
 	{
 		if (SendGetFileListMessage())
 		{
diff --git a/Engine/Source/Runtime/StorageServerClient/Private/StorageServerPlatformFile.h b/Engine/Source/Runtime/StorageServerClient/Private/StorageServerPlatformFile.h
--- a/Engine/Source/Runtime/StorageServerClient/Private/StorageServerPlatformFile.h
+++ b/Engine/Source/Runtime/StorageServerClient/Private/StorageServerPlatformFile.h
@@ -119,6 +119,7 @@ private:
 	void OnCookOnTheFlyMessage(const UE::Cook::FCookOnTheFlyMessage& Message);
 #endif
 	TUniquePtr<FArchive> TryFindProjectStoreMarkerFile(IPlatformFile* Inner) const;
+	FAnsiString MakeBaseURI();
 
 	TSet<FName> ExcludedNonServerExtensions;
 	IPlatformFile* LowerLevel = nullptr;
@@ -131,6 +132,7 @@ private:
 	FStorageServerFileSystemTOC ServerToc;
 	FString ServerProject;
 	FString ServerPlatform;
+	FString BaseURI;
 	mutable TArray<FString> HostAddrs;
 	mutable uint16 HostPort = 8558;
 };
diff --git a/Engine/Source/Runtime/StorageServerClient/Private/StorageSocketConnectionBackend.cpp b/Engine/Source/Runtime/StorageServerClient/Private/StorageSocketConnectionBackend.cpp
--- a/Engine/Source/Runtime/StorageServerClient/Private/StorageSocketConnectionBackend.cpp
+++ b/Engine/Source/Runtime/StorageServerClient/Private/StorageSocketConnectionBackend.cpp
@@ -207,7 +207,7 @@ bool FStorageSocketConnectionBackend::InitializeInternal(TArrayView<const FStrin
 int32 FStorageSocketConnectionBackend::HandshakeRequest(TArrayView<const TSharedPtr<FInternetAddr>> HostAddresses)
 {
 	TAnsiStringBuilder<256> ResourceBuilder;
-	ResourceBuilder.Append(OplogPath);
+	ResourceBuilder.Append(BaseURI);
 
 	TArray<TSharedPtr<FInternetAddr>> SortedAddresses;
 	SortHostAddressesByLocalSubnet(HostAddresses, SortedAddresses);
diff --git a/Engine/Source/Runtime/StorageServerClient/Public/StorageServerConnection.h b/Engine/Source/Runtime/StorageServerClient/Public/StorageServerConnection.h
--- a/Engine/Source/Runtime/StorageServerClient/Public/StorageServerConnection.h
+++ b/Engine/Source/Runtime/StorageServerClient/Public/StorageServerConnection.h
@@ -5,6 +5,7 @@
 #include "CoreTypes.h"
 #include "Compression/CompressedBuffer.h"
 #include "Containers/UnrealString.h"
+#include "Containers/AnsiString.h"
 #include "Misc/StringBuilder.h"
 #include "Memory/MemoryFwd.h"
 #include "Templates/SharedPointer.h"
@@ -218,7 +219,7 @@ public:
 	FStorageConnectionBackend(FStorageServerConnection& InOwner);
 	virtual ~FStorageConnectionBackend() = default;
 
-	bool Initialize(TArrayView<const FString> InHostAddresses, int32 InPort, const TCHAR* InProjectNameOverride, const TCHAR* InPlatformNameOverride);
+	bool Initialize(TArrayView<const FString> InHostAddresses, int32 InPort, const FAnsiStringView& InBaseURI);
 
 	virtual IStorageConnectionSocket* AcquireSocketFromPool() = 0;
 	virtual IStorageConnectionSocket* AcquireNewSocket(float TimeoutSeconds = -1.f) = 0;
@@ -226,12 +227,11 @@ public:
 	virtual FString GetHostName() = 0;
 	
 protected:
-	void InitOplog(const TCHAR* InProjectNameOverride, const TCHAR* InPlatformNameOverride);
 	virtual bool InitializeInternal(TArrayView<const FString> InHostAddresses, int32 Port) = 0;
 
 protected:
 	FStorageServerConnection& Owner;
-	TAnsiStringBuilder<1024> OplogPath;
+	FAnsiString BaseURI;
 };
 
 class UDebugStorageServerConnection;
@@ -242,7 +242,7 @@ public:
 	STORAGESERVERCLIENT_API FStorageServerConnection();
 	STORAGESERVERCLIENT_API ~FStorageServerConnection();
 
-	STORAGESERVERCLIENT_API bool Initialize(TArrayView<const FString> HostAddresses, int32 Port, const TCHAR* ProjectNameOverride = nullptr, const TCHAR* PlatformNameOverride = nullptr);
+	STORAGESERVERCLIENT_API bool Initialize(TArrayView<const FString> HostAddresses, int32 Port, const FAnsiStringView& InBaseURI);
 
 	STORAGESERVERCLIENT_API void PackageStoreRequest(TFunctionRef<void(FPackageStoreEntryResource&&)> Callback);
 	STORAGESERVERCLIENT_API void FileManifestRequest(TFunctionRef<void(FIoChunkId Id, FStringView Path, int64 RawSize)> Callback);
@@ -265,7 +265,7 @@ private:
 	void AddTimingInstance(double duration, uint64 bytes);
 
 	ISocketSubsystem& SocketSubsystem;
-	TAnsiStringBuilder<1024> OplogPath;
+	FAnsiString BaseURI;
 	TSharedPtr<FInternetAddr> ServerAddr;
 	TAnsiStringBuilder<1024> Hostname;
 	TUniquePtr<FStorageConnectionBackend> ConnectionBackend;
