Refactor Zen storage module to be able to use other HTTP clients

diff --git a/Engine/Source/Runtime/Core/Public/IO/IoBuffer.h b/Engine/Source/Runtime/Core/Public/IO/IoBuffer.h
--- a/Engine/Source/Runtime/Core/Public/IO/IoBuffer.h
+++ b/Engine/Source/Runtime/Core/Public/IO/IoBuffer.h
@@ -90,8 +90,8 @@ private:
 
 		//
 
-		void	SetDataAndSize(const uint8* InData, uint64 InSize);
-		void	SetSize(uint64 InSize);
+		CORE_API void SetDataAndSize(const uint8* InData, uint64 InSize);
+		CORE_API void SetSize(uint64 InSize);
 
 		void	MakeOwned();
 
diff --git a/Engine/Source/Runtime/StorageServerClient/Private/BuiltInHttpClient/BuiltInHttpClient.cpp b/Engine/Source/Runtime/StorageServerClient/Private/BuiltInHttpClient/BuiltInHttpClient.cpp
new file mode 100644
--- /dev/null
+++ b/Engine/Source/Runtime/StorageServerClient/Private/BuiltInHttpClient/BuiltInHttpClient.cpp
@@ -0,0 +1,202 @@
+// Copyright Epic Games, Inc. All Rights Reserved.
+
+#include "BuiltInHttpClient.h"
+#include "StorageServerConnection.h"
+
+#if !UE_BUILD_SHIPPING
+
+static inline FAnsiStringView GetMimeTypeString(EStorageServerContentType ContentType)
+{
+	switch (ContentType)
+	{
+		case EStorageServerContentType::CbObject:
+			return ANSITEXTVIEW("application/ue-x-cb");
+		case EStorageServerContentType::Binary:
+			return ANSITEXTVIEW("application/octet-stream");
+		case EStorageServerContentType::CompressedBinary:
+			return ANSITEXTVIEW("application/x-ue-comp");
+		default:
+			return ANSITEXTVIEW("unknown");
+	};
+};
+
+static inline EStorageServerContentType GetMimeType(const FAnsiStringView& ContentType)
+{
+	if (ContentType == ANSITEXTVIEW("application/octet-stream"))
+	{
+		return EStorageServerContentType::Binary;
+	}
+	else if (ContentType == ANSITEXTVIEW("application/x-ue-comp"))
+	{
+		return EStorageServerContentType::CompressedBinary;
+	}
+	else if (ContentType == ANSITEXTVIEW("application/x-ue-cb"))
+	{
+		return EStorageServerContentType::CbObject;
+	}
+	else
+	{
+		return EStorageServerContentType::Unknown;
+	}
+};
+
+FBuiltInHttpClient::FBuiltInHttpClient(TUniquePtr<IBuiltInHttpClientSocketPool> InSocketPool, FString InHostname)
+	: SocketPool(MoveTemp(InSocketPool))
+	, Hostname(InHostname)
+{
+}
+
+IStorageServerHttpClient::FResult FBuiltInHttpClient::RequestSync(
+	FAnsiStringView Url,
+	EStorageServerContentType Accept,
+	FAnsiStringView Verb,
+	TOptional<FIoBuffer> OptPayload,
+	TOptional<FIoBuffer> OptDestination,
+	float TimeoutSeconds,
+	const bool bReportErrors
+)
+{
+	const bool bHasPayload = OptPayload.IsSet() && OptPayload->GetSize() > 0;
+
+	TAnsiStringBuilder<1024> HeaderBuffer;
+
+	HeaderBuffer
+		<< Verb << " " << Url << " HTTP/1.1\r\n" 
+		<< "Host: " << TCHAR_TO_ANSI(*Hostname) << "\r\n"
+		<< "Connection: Keep-Alive\r\n";
+	if (Accept != EStorageServerContentType::Unknown)
+	{
+		HeaderBuffer
+			<< "Accept: " << GetMimeTypeString(Accept) << "\r\n";
+	}
+	if (bHasPayload)
+	{
+		HeaderBuffer
+			<< "Content-Length: " << OptPayload->GetSize() << "\r\n";
+	}
+	HeaderBuffer << "\r\n";
+
+	IBuiltInHttpClientSocket* Socket = nullptr;
+
+	bool bSendOk = false;
+
+	const int32 AttemptCount = TimeoutSeconds <= 0.0f ? 10 : 1; 
+	for (int32 Attempt = 0; Attempt < AttemptCount; Attempt++)
+	{
+		if ((Socket = SocketPool->AcquireSocket(TimeoutSeconds)) == nullptr)
+		{
+			continue;
+		}
+
+		if (Socket->Send(reinterpret_cast<const uint8*>(HeaderBuffer.GetData()), HeaderBuffer.Len()) &&
+			(!bHasPayload || Socket->Send(OptPayload->GetData(), OptPayload->GetSize())))
+		{
+			bSendOk = true;
+			break;
+		}
+
+		SocketPool->ReleaseSocket(Socket, false);
+	}
+
+	if (!bSendOk)
+	{
+		if (bReportErrors)
+		{
+			UE_LOG(LogStorageServerConnection, Fatal, TEXT("Failed sending request to storage server."));
+		}
+
+		return {TIoStatusOr<FIoBuffer>(FIoStatus(EIoErrorCode::WriteError)), EStorageServerContentType::Unknown};
+	}
+
+	bool bRecvOk = true;
+
+	uint8 ReadResponseLineBuffer[1024];
+	auto ReadResponseLine = [&ReadResponseLineBuffer, &Socket, &bRecvOk]() -> FAnsiStringView
+	{
+		for (;;)
+		{
+			uint64 BytesRead;
+			Socket->Recv(ReadResponseLineBuffer, 1024, BytesRead, ESocketReceiveFlags::Peek);
+			FAnsiStringView ResponseView(reinterpret_cast<const ANSICHAR*>(ReadResponseLineBuffer), BytesRead);
+			int32 LineEndIndex;
+			if (ResponseView.FindChar('\r', LineEndIndex) && BytesRead >= LineEndIndex + 2)
+			{
+				check(ResponseView[LineEndIndex + 1] == '\n');
+				if (!Socket->Recv(ReadResponseLineBuffer, LineEndIndex + 2, BytesRead, ESocketReceiveFlags::None))
+				{
+					bRecvOk = false;
+				}
+				check(BytesRead == LineEndIndex + 2);
+				return ResponseView.Left(LineEndIndex);
+			}
+		}
+	};
+
+	int64 ContentLength = 0;
+	int32 StatusCode = 500;
+	EStorageServerContentType ContentType = EStorageServerContentType::Unknown;
+	while (true)
+	{
+		FAnsiStringView ResponseLine = ReadResponseLine();
+
+		if (ResponseLine.IsEmpty())
+		{
+			break;
+		}
+		else if (ResponseLine.StartsWith("HTTP/1.1 "))
+		{
+			StatusCode = TCString<ANSICHAR>::Atoi64(ResponseLine.GetData() + 9);
+		}
+		else if (ResponseLine.StartsWith("Content-Length: "))
+		{
+			ContentLength = FMath::Max(0, TCString<ANSICHAR>::Atoi64(ResponseLine.GetData() + 16));
+		}
+		else if (ResponseLine.StartsWith("Content-Type: "))
+		{
+			ContentType = GetMimeType(ResponseLine.RightChop(14));
+		}
+	}
+
+	const bool bIsOk = StatusCode < 300; // we don't handle redirects here, so everything >=300 is considered an error
+	const EIoErrorCode ErrorCode = bIsOk ? EIoErrorCode::Ok : (StatusCode == 404 ? EIoErrorCode::NotFound : EIoErrorCode::Unknown);
+
+	TOptional<FIoBuffer> ResponsePayload = ContentLength > 0
+		? (OptDestination.IsSet() && ((int64)OptDestination->GetSize() >= ContentLength) ? OptDestination.GetValue() : FIoBuffer(ContentLength))
+		: TOptional<FIoBuffer>();
+
+	if (ContentLength)
+	{
+		int64 TotalBytesRead = 0;
+		while (TotalBytesRead < ContentLength)
+		{
+			uint64 BytesRead = 0;
+			if (!Socket->Recv(ResponsePayload->GetData() + TotalBytesRead, ContentLength - TotalBytesRead, BytesRead, ESocketReceiveFlags::WaitAll))
+			{
+				bRecvOk = false;
+				break;
+			}
+
+			TotalBytesRead += (int64)BytesRead;
+		}
+
+		ResponsePayload->SetSize(TotalBytesRead);
+	}
+
+	const bool bHasResponsePayload = ResponsePayload.IsSet() && ResponsePayload->GetSize() > 0;
+
+	const bool bKeepAlive = bSendOk && bRecvOk;
+	SocketPool->ReleaseSocket(Socket, bKeepAlive);
+
+	if (bIsOk)
+	{
+		return {TIoStatusOr<FIoBuffer>(MoveTemp(ResponsePayload.GetValue())), ContentType};
+	}
+	else
+	{
+		FString ErrorMessage = bHasResponsePayload ? FString(ResponsePayload->GetSize(), ANSI_TO_TCHAR(reinterpret_cast<ANSICHAR*>(ResponsePayload->GetData()))) : TEXT("Unknown error");
+
+		return {TIoStatusOr<FIoBuffer>(FIoStatus(ErrorCode, ErrorMessage)), ContentType};
+	}
+}
+
+#endif
\ No newline at end of file
diff --git a/Engine/Source/Runtime/StorageServerClient/Private/BuiltInHttpClient/BuiltInHttpClient.h b/Engine/Source/Runtime/StorageServerClient/Private/BuiltInHttpClient/BuiltInHttpClient.h
new file mode 100644
--- /dev/null
+++ b/Engine/Source/Runtime/StorageServerClient/Private/BuiltInHttpClient/BuiltInHttpClient.h
@@ -0,0 +1,65 @@
+// Copyright Epic Games, Inc. All Rights Reserved.
+
+#pragma once
+
+#include "StorageServerHttpClient.h"
+#include "SocketTypes.h"
+
+#if !UE_BUILD_SHIPPING
+
+class IBuiltInHttpClientSocket
+{
+public:
+	virtual ~IBuiltInHttpClientSocket() = default;
+
+	virtual bool Send(const uint8* Data, const uint64 DataSize)                                                      = 0;
+	virtual bool Recv(uint8* Data, const uint64 DataSize, uint64& BytesRead, ESocketReceiveFlags::Type ReceiveFlags) = 0;
+	virtual bool HasPendingData(uint64& PendingDataSize) const                                                       = 0;
+	virtual void Close()                                                                                             = 0;
+};
+
+class IBuiltInHttpClientSocketPool
+{
+public:
+	virtual ~IBuiltInHttpClientSocketPool() = default;
+
+	virtual IBuiltInHttpClientSocket* AcquireSocket(float TimeoutSeconds = -1.f) = 0;
+	virtual void ReleaseSocket(IBuiltInHttpClientSocket* Socket, bool bKeepAlive) = 0;
+};
+
+class FBuiltInHttpClient : public IStorageServerHttpClient
+{
+public:
+	FBuiltInHttpClient(TUniquePtr<IBuiltInHttpClientSocketPool> InSocketPool, FString InHostname);
+	virtual ~FBuiltInHttpClient() override = default;
+
+	virtual FResult RequestSync(
+		FAnsiStringView Url,
+		EStorageServerContentType Accept,
+		FAnsiStringView Verb,
+		TOptional<FIoBuffer> OptPayload,
+		TOptional<FIoBuffer> OptDestination,
+		float TimeoutSeconds,
+		const bool bReportErrors
+	) override;
+	
+	virtual void RequestAsync(
+		FResultCallback&& Callback,
+		FAnsiStringView Url,
+		EStorageServerContentType Accept,
+		FAnsiStringView Verb,
+		TOptional<FIoBuffer> OptPayload,
+		TOptional<FIoBuffer> OptDestination,
+		float TimeoutSeconds,
+		const bool bReportErrors
+	) override
+	{
+		// This HTTP client doesn't support async requests
+		Callback(RequestSync(Url, Accept, Verb, OptPayload, OptDestination, TimeoutSeconds, bReportErrors));
+	}
+private:
+	TUniquePtr<IBuiltInHttpClientSocketPool> SocketPool;
+	FString Hostname;
+};
+
+#endif
diff --git a/Engine/Source/Runtime/StorageServerClient/Private/BuiltInHttpClient/BuiltInHttpClientFSocket.cpp b/Engine/Source/Runtime/StorageServerClient/Private/BuiltInHttpClient/BuiltInHttpClientFSocket.cpp
new file mode 100644
--- /dev/null
+++ b/Engine/Source/Runtime/StorageServerClient/Private/BuiltInHttpClient/BuiltInHttpClientFSocket.cpp
@@ -0,0 +1,148 @@
+// Copyright Epic Games, Inc. All Rights Reserved.
+
+#include "BuiltInHttpClientFSocket.h"
+#include "StorageServerConnection.h"
+#include "Sockets.h"
+#include "SocketSubsystem.h"
+#include "Misc/ScopeExit.h"
+
+#if !UE_BUILD_SHIPPING
+
+FBuiltInHttpClientFSocket::FBuiltInHttpClientFSocket(FSocket* InSocket)
+	: Socket(InSocket)
+{
+}
+
+FBuiltInHttpClientFSocket::~FBuiltInHttpClientFSocket()
+{
+	if (Socket)
+	{
+		Socket->Close();
+		delete Socket;
+		Socket = nullptr;
+	}
+}
+
+bool FBuiltInHttpClientFSocket::Send(const uint8* Data, const uint64 DataSize)
+{
+	if (!Socket)
+	{
+		return false;
+	}
+
+	int32 TotalBytesSent = 0;
+	while (TotalBytesSent < (int64)DataSize)
+	{
+		int32 BytesSent = 0;
+		if (!Socket->Send(Data, DataSize, BytesSent))
+		{
+			return false;
+		}
+
+		TotalBytesSent += BytesSent;
+	}
+
+	return true;
+}
+
+bool FBuiltInHttpClientFSocket::Recv(uint8* Data, const uint64 DataSize, uint64& BytesRead, ESocketReceiveFlags::Type ReceiveFlags)
+{
+	if (!Socket)
+	{
+		return false;
+	}
+
+	int32 ReadBytes = 0;
+	if (!Socket->Recv(Data, DataSize, ReadBytes, ReceiveFlags))
+	{
+		return false;
+	}
+
+	BytesRead = ReadBytes;
+	return true;
+}
+
+bool FBuiltInHttpClientFSocket::HasPendingData(uint64& PendingDataSize) const
+{
+	uint32 PendingData;
+	bool bRes = Socket->HasPendingData(PendingData);
+
+	PendingDataSize = PendingData;
+	return bRes;
+}
+
+void FBuiltInHttpClientFSocket::Close()
+{
+	Socket->Close();
+}
+
+FBuiltInHttpClientFSocketPool::FBuiltInHttpClientFSocketPool(TSharedPtr<FInternetAddr> InServerAddr, ISocketSubsystem& InSocketSubsystem)
+	: ServerAddr(InServerAddr)
+	, SocketSubsystem(InSocketSubsystem)
+{
+}
+
+FBuiltInHttpClientFSocketPool::~FBuiltInHttpClientFSocketPool()
+{
+	IBuiltInHttpClientSocket* Socket = nullptr;
+	while ((Socket = SocketPool.Pop()) != nullptr)
+	{
+		delete Socket;
+	}
+}
+
+IBuiltInHttpClientSocket* FBuiltInHttpClientFSocketPool::AcquireSocket(float TimeoutSeconds)
+{
+	IBuiltInHttpClientSocket* SocketWrapper = SocketPool.Pop();
+	if (SocketWrapper)
+	{
+		return SocketWrapper;
+	}
+
+	FSocket* Socket = SocketSubsystem.CreateSocket(NAME_Stream, TEXT("StorageServer"), ServerAddr->GetProtocolType());
+
+	if (TimeoutSeconds > 0.0f)
+	{
+		Socket->SetNonBlocking(true);
+
+		ON_SCOPE_EXIT
+		{
+			Socket->SetNonBlocking(false);
+		};
+
+		if (Socket->Connect(*ServerAddr) && Socket->Wait(ESocketWaitConditions::WaitForWrite, FTimespan::FromSeconds(TimeoutSeconds)))
+		{
+			return new FBuiltInHttpClientFSocket(Socket);
+		}
+	}
+	else
+	{
+		if (Socket->Connect(*ServerAddr))
+		{
+			return new FBuiltInHttpClientFSocket(Socket);
+		}
+	}
+
+	delete Socket;
+	return nullptr;
+}
+
+void FBuiltInHttpClientFSocketPool::ReleaseSocket(IBuiltInHttpClientSocket* Socket, bool bKeepAlive)
+{
+	uint64 PendingDataSize = 0;
+	if (bKeepAlive && !Socket->HasPendingData(PendingDataSize))
+	{
+		SocketPool.Push(Socket);
+	}
+	else
+	{
+		if (PendingDataSize > 0)
+		{
+			UE_LOG(LogStorageServerConnection, Fatal, TEXT("Socket was not fully drained"));
+		}
+
+		delete Socket;
+	}
+}
+
+#endif
\ No newline at end of file
diff --git a/Engine/Source/Runtime/StorageServerClient/Private/BuiltInHttpClient/BuiltInHttpClientFSocket.h b/Engine/Source/Runtime/StorageServerClient/Private/BuiltInHttpClient/BuiltInHttpClientFSocket.h
new file mode 100644
--- /dev/null
+++ b/Engine/Source/Runtime/StorageServerClient/Private/BuiltInHttpClient/BuiltInHttpClientFSocket.h
@@ -0,0 +1,42 @@
+// Copyright Epic Games, Inc. All Rights Reserved.
+
+#pragma once
+
+#include "BuiltInHttpClient.h"
+#include "Containers/LockFreeList.h"
+
+#if !UE_BUILD_SHIPPING
+
+class FSocket;
+
+class FBuiltInHttpClientFSocket : public IBuiltInHttpClientSocket
+{
+public:
+	FBuiltInHttpClientFSocket(FSocket* InSocket);
+	virtual ~FBuiltInHttpClientFSocket() override;
+
+	virtual bool Send(const uint8* Data, const uint64 DataSize) override;
+	virtual bool Recv(uint8* Data, const uint64 DataSize, uint64& BytesRead, ESocketReceiveFlags::Type ReceiveFlags) override;
+	virtual bool HasPendingData(uint64& PendingDataSize) const override;
+	virtual void Close() override;
+
+private:
+	FSocket* Socket;
+};
+
+class FBuiltInHttpClientFSocketPool : public IBuiltInHttpClientSocketPool
+{
+public:
+	FBuiltInHttpClientFSocketPool(TSharedPtr<FInternetAddr> InServerAddr, ISocketSubsystem& InSocketSubsystem);
+	virtual ~FBuiltInHttpClientFSocketPool() override;
+
+	virtual IBuiltInHttpClientSocket* AcquireSocket(float TimeoutSeconds = -1.f) override;
+	virtual void ReleaseSocket(IBuiltInHttpClientSocket* Socket, bool bKeepAlive) override;
+
+private:
+	TSharedPtr<FInternetAddr> ServerAddr;
+	ISocketSubsystem& SocketSubsystem;
+	TLockFreePointerListUnordered<IBuiltInHttpClientSocket, PLATFORM_CACHE_LINE_SIZE> SocketPool;
+};
+
+#endif
\ No newline at end of file
diff --git a/Engine/Source/Runtime/StorageServerClient/Private/BuiltInHttpClient/BuiltInHttpClientPlatformSocket.cpp b/Engine/Source/Runtime/StorageServerClient/Private/BuiltInHttpClient/BuiltInHttpClientPlatformSocket.cpp
new file mode 100644
--- /dev/null
+++ b/Engine/Source/Runtime/StorageServerClient/Private/BuiltInHttpClient/BuiltInHttpClientPlatformSocket.cpp
@@ -0,0 +1,216 @@
+// Copyright Epic Games, Inc. All Rights Reserved.
+
+#include "BuiltInHttpClientPlatformSocket.h"
+#include "HAL/PlatformProcess.h"
+
+#if !UE_BUILD_SHIPPING
+
+DEFINE_LOG_CATEGORY_STATIC(LogStorageServerPlatformBackend, Log, All);
+
+FBuiltInHttpClientPlatformSocket::FBuiltInHttpClientPlatformSocket(IPlatformHostCommunication* InCommunication, IPlatformHostSocketPtr InSocket, int32 InProtocolNumber)
+	: Communication(InCommunication)
+	, Socket(InSocket)
+	, ConnectionBuffer(1024 * 64)
+	, ProtocolNumber(InProtocolNumber)
+{
+}
+
+FBuiltInHttpClientPlatformSocket::~FBuiltInHttpClientPlatformSocket()
+{
+	Close();
+}
+
+bool FBuiltInHttpClientPlatformSocket::Send(const uint8* Data, const uint64 DataSize)
+{
+	if (!Socket || Socket->GetState() != IPlatformHostSocket::EConnectionState::Connected)
+	{
+		return false;
+	}
+
+	const bool bSuccess = Socket->Send(Data, DataSize) == IPlatformHostSocket::EResultNet::Ok;
+	return bSuccess;
+}
+
+bool FBuiltInHttpClientPlatformSocket::Recv(uint8* Data, const uint64 DataSize, uint64& BytesRead, ESocketReceiveFlags::Type ReceiveFlags)
+{
+	if (!Socket || Socket->GetState() != IPlatformHostSocket::EConnectionState::Connected)
+	{
+		return false;
+	}
+
+	if (ConnectionBuffer.GetCapacity() < DataSize)
+	{
+		UE_LOG(LogStorageServerPlatformBackend, Display, TEXT("ConnectionBuffer capacity is lower than requested data read (%d vs %d)"), ConnectionBuffer.GetCapacity(), DataSize);
+	}
+
+	BytesRead = 0;
+	if (ConnectionBuffer.IsEmpty())
+	{
+		uint8 Buffer[1024];
+		if (Socket->Receive(Buffer, sizeof(Buffer), BytesRead, IPlatformHostSocket::EReceiveFlags::DontWait) != IPlatformHostSocket::EResultNet::Ok)
+		{
+			return false;
+		}
+
+		if (!ConnectionBuffer.Put(Buffer, BytesRead))
+		{
+			UE_LOG(LogStorageServerPlatformBackend, Display, TEXT("Couldn't fit the received data in the connection buffer"));
+		}
+	}
+
+	if (ReceiveFlags == ESocketReceiveFlags::Peek)
+	{
+		ConnectionBuffer.Peek(Data, DataSize, BytesRead);
+	}
+	else if (ReceiveFlags == ESocketReceiveFlags::WaitAll)
+	{
+		ConnectionBuffer.Consume(Data, DataSize, BytesRead);
+		uint64 TotalBytesRead = BytesRead;
+
+		while (TotalBytesRead < DataSize)
+		{
+			const uint64 BytesToRead = DataSize - TotalBytesRead;
+			if (Socket->Receive(Data + TotalBytesRead, BytesToRead, BytesRead, IPlatformHostSocket::EReceiveFlags::WaitAll) != IPlatformHostSocket::EResultNet::Ok)
+			{
+				return false;
+			}
+
+			TotalBytesRead += BytesRead;
+
+			if (TotalBytesRead > DataSize)
+			{
+				UE_LOG(LogStorageServerPlatformBackend, Display, TEXT("Exceeded what was supposed to be downloaded"));
+			}
+		}
+
+		BytesRead = TotalBytesRead;
+	}
+	else if (ReceiveFlags == ESocketReceiveFlags::None)
+	{
+		ConnectionBuffer.Consume(Data, DataSize, BytesRead);
+	}
+
+	return true;
+}
+
+bool FBuiltInHttpClientPlatformSocket::HasPendingData(uint64& PendingDataSize) const
+{
+	PendingDataSize = ConnectionBuffer.GetSize();
+	return PendingDataSize > 0;
+}
+
+void FBuiltInHttpClientPlatformSocket::Close()
+{
+	Communication->CloseConnection(Socket);
+}
+
+FBuiltInHttpClientPlatformSocketPool::FBuiltInHttpClientPlatformSocketPool(const FString InAddress)
+	: Address(InAddress)
+{
+	Communication = &FPlatformMisc::GetPlatformHostCommunication();
+
+	UsedSockets.Init(false, 10); // TODO add host communication to get available amount of socket connection
+}
+
+FBuiltInHttpClientPlatformSocketPool::~FBuiltInHttpClientPlatformSocketPool()
+{
+	IBuiltInHttpClientSocket* Socket = nullptr;
+	while ((Socket = SocketPool.Pop()) != nullptr)
+	{
+		delete Socket;
+	}
+}
+
+IBuiltInHttpClientSocket* FBuiltInHttpClientPlatformSocketPool::AcquireSocket(float TimeoutSeconds)
+{
+	IBuiltInHttpClientSocket* SocketWrapper = SocketPool.Pop();
+	if (SocketWrapper)
+	{
+		return SocketWrapper;
+	}
+
+	int32 ProtocolNumber = -1;
+	while (ProtocolNumber == -1)
+	{
+		FScopeLock Lock(&UsedSocketsCS);
+		ProtocolNumber = UsedSockets.Find(false);
+
+		if (ProtocolNumber == -1)
+		{
+			// All sockets are in use, and we have limited amount of sockets we could use
+			UsedSocketsCV.Wait(UsedSocketsCS);
+		}
+		else
+		{
+			UsedSockets[ProtocolNumber] = true;
+			break;
+		}
+	}
+
+	IPlatformHostSocket::EConnectionState ConnectionState = IPlatformHostSocket::EConnectionState::Unknown;
+
+	// TODO use address to specify which device to connect to
+	IPlatformHostSocketPtr PlatformSocket = Communication->OpenConnection(ProtocolNumber, *FString::Printf(TEXT("PlatformSocket %d"), ProtocolNumber));
+	if (PlatformSocket)
+	{
+		float WaitingFor = 0.f;
+		const float SleepTime = 0.01f;
+		ConnectionState = PlatformSocket->GetState();
+		while (ConnectionState == IPlatformHostSocket::EConnectionState::Created)
+		{
+			if ((TimeoutSeconds != -1.f) && (WaitingFor > TimeoutSeconds))
+			{
+				UE_LOG(LogStorageServerPlatformBackend, Error, TEXT("Platform connection timed out"));
+				break;
+			}
+
+			FPlatformProcess::Sleep(SleepTime);
+			WaitingFor += SleepTime;
+			ConnectionState = PlatformSocket->GetState();
+		}
+	}
+
+	if (ConnectionState == IPlatformHostSocket::EConnectionState::Connected)
+	{
+		return new FBuiltInHttpClientPlatformSocket(Communication, PlatformSocket, ProtocolNumber);
+	}
+	else
+	{
+		if (PlatformSocket)
+		{
+			Communication->CloseConnection(PlatformSocket);
+		}
+
+		FScopeLock Lock(&UsedSocketsCS);
+		UsedSockets[ProtocolNumber] = false;
+		return nullptr;
+	}
+}
+
+void FBuiltInHttpClientPlatformSocketPool::ReleaseSocket(IBuiltInHttpClientSocket* Socket, bool bKeepAlive)
+{
+	uint64 PendingDataSize = 0;
+	if (bKeepAlive && !Socket->HasPendingData(PendingDataSize))
+	{
+		SocketPool.Push(Socket);
+	}
+	else
+	{
+		const int32 ProtocolNumber = static_cast<FBuiltInHttpClientPlatformSocket*>(Socket)->GetProtocolNumber();
+
+		if (PendingDataSize > 0)
+		{
+			UE_LOG(LogStorageServerPlatformBackend, Fatal, TEXT("Socket was not fully drained"));
+		}
+
+		delete Socket;
+
+		{
+			FScopeLock Lock(&UsedSocketsCS);
+			UsedSockets[ProtocolNumber] = false;
+			UsedSocketsCV.NotifyOne();
+		}
+	}
+}
+
+#endif
\ No newline at end of file
diff --git a/Engine/Source/Runtime/StorageServerClient/Private/BuiltInHttpClient/BuiltInHttpClientPlatformSocket.h b/Engine/Source/Runtime/StorageServerClient/Private/BuiltInHttpClient/BuiltInHttpClientPlatformSocket.h
new file mode 100644
--- /dev/null
+++ b/Engine/Source/Runtime/StorageServerClient/Private/BuiltInHttpClient/BuiltInHttpClientPlatformSocket.h
@@ -0,0 +1,171 @@
+// Copyright Epic Games, Inc. All Rights Reserved.
+
+#pragma once
+
+#include "BuiltInHttpClient.h"
+#include "Experimental/Async/ConditionVariable.h"
+#include "GenericPlatform/GenericPlatformHostCommunication.h"
+#include "GenericPlatform/GenericPlatformHostSocket.h"
+#include "Containers/LockFreeList.h"
+
+#if !UE_BUILD_SHIPPING
+
+class FConnectionCircularBuffer
+{
+public:
+	explicit FConnectionCircularBuffer(uint64 InCapacity)
+		: CapacityMask(FMath::RoundUpToPowerOfTwo(InCapacity) - 1)
+	{
+		Allocation = new uint8[GetCapacity()];
+	}
+
+	~FConnectionCircularBuffer()
+	{
+		delete[] Allocation;
+	}
+
+	uint8 operator[](uint64 Index) const
+	{
+		return const_cast<FConnectionCircularBuffer*>(this)->operator[](Index);
+	}
+
+	uint8& operator[](uint64 Index)
+	{
+		check(Index < GetSize());
+		uint64 DataOffset = (Index + Tail) & (CapacityMask);
+		return Allocation[DataOffset];
+	}
+
+	uint64 GetSize() const
+	{
+		return Size;
+	}
+
+	uint64 GetCapacity() const
+	{
+		return CapacityMask + 1;
+	}
+
+	bool IsEmpty() const
+	{
+		return Size == 0;
+	}
+
+	void Clear()
+	{
+		Tail = 0;
+		Head = 0;
+		Size = 0;
+	}
+
+	uint64 SpaceLeft() const
+	{
+		return GetCapacity() - Size;
+	}
+
+	void Peek(uint8* Data, const uint64 DataSize, uint64& OutSize)
+	{
+		OutSize = FMath::Min(DataSize, Size);
+
+		if (OutSize + Tail > GetCapacity())
+		{
+			uint64 LenToEnd = GetCapacity() - Tail;
+			uint64 LenFromBegin = OutSize - LenToEnd;
+			FMemory::Memcpy(Data, Allocation + Tail, LenToEnd);
+			FMemory::Memcpy(Data + LenToEnd, Allocation, LenFromBegin);
+		}
+		else
+		{
+			FMemory::Memcpy(Data, Allocation + Tail, OutSize);
+		}
+	}
+
+	void Consume(uint8* Data, const uint64 DataSize, uint64& OutSize)
+	{
+		Peek(Data, DataSize, OutSize);
+		Tail = (Tail + OutSize) & CapacityMask;
+
+		Size -= OutSize;
+
+		if (Size == 0)
+		{
+			Head = 0;
+			Tail = 0;
+		}
+	}
+
+	bool Put(uint8* Data, uint64 DataSize)
+	{
+		if (DataSize > SpaceLeft())
+		{
+			return false;
+		}
+
+		if (DataSize + Head > GetCapacity())
+		{
+			uint64 LenToEnd = GetCapacity() - Head;
+			uint64 LenFromBegin = DataSize - LenToEnd;
+			FMemory::Memcpy(Allocation + Head, Data, LenToEnd);
+			FMemory::Memcpy(Allocation, Data + LenToEnd, LenFromBegin);
+		}
+		else
+		{
+			FMemory::Memcpy(Allocation + Head, Data, DataSize);
+		}
+
+		Head = (Head + DataSize) & (CapacityMask);
+
+		Size += DataSize;
+
+		return true;
+	}
+
+private:
+	uint8* Allocation{ nullptr };
+	uint64 CapacityMask{ 0 };
+	uint64 Size{ 0 };
+	uint64 Head{ 0 };
+	uint64 Tail{ 0 };
+};
+
+class FBuiltInHttpClientPlatformSocket : public IBuiltInHttpClientSocket
+{
+public:
+	FBuiltInHttpClientPlatformSocket(IPlatformHostCommunication* InCommunication, IPlatformHostSocketPtr InSocket, int32 InProtocolNumber);
+	virtual ~FBuiltInHttpClientPlatformSocket() override;
+
+	virtual bool Send(const uint8* Data, const uint64 DataSize) override;
+	virtual bool Recv(uint8* Data, const uint64 DataSize, uint64& BytesRead, ESocketReceiveFlags::Type ReceiveFlags) override;
+	virtual bool HasPendingData(uint64& PendingDataSize) const override;
+	virtual void Close() override;
+
+	int32 GetProtocolNumber() const { return ProtocolNumber; }
+
+private:
+	IPlatformHostCommunication* Communication;
+	IPlatformHostSocketPtr Socket;
+	FConnectionCircularBuffer ConnectionBuffer;
+	const int32 ProtocolNumber;
+};
+
+class FBuiltInHttpClientPlatformSocketPool : public IBuiltInHttpClientSocketPool
+{
+public:
+	FBuiltInHttpClientPlatformSocketPool(const FString InAddress);
+	virtual ~FBuiltInHttpClientPlatformSocketPool() override;
+
+	virtual IBuiltInHttpClientSocket* AcquireSocket(float TimeoutSeconds = -1.f) override;
+	virtual void ReleaseSocket(IBuiltInHttpClientSocket* Socket, bool bKeepAlive) override;
+
+private:
+	const FString Address;
+	IPlatformHostCommunication* Communication = nullptr;
+
+	TLockFreePointerListUnordered<IBuiltInHttpClientSocket, PLATFORM_CACHE_LINE_SIZE> SocketPool;
+
+	FCriticalSection UsedSocketsCS;
+	UE::FConditionVariable UsedSocketsCV;
+	TBitArray<> UsedSockets; // bitset to keep track of used sockets in the pool
+};
+
+#endif
\ No newline at end of file
diff --git a/Engine/Source/Runtime/StorageServerClient/Private/DebugStorageServerConnection.cpp b/Engine/Source/Runtime/StorageServerClient/Private/DebugStorageServerConnection.cpp
--- a/Engine/Source/Runtime/StorageServerClient/Private/DebugStorageServerConnection.cpp
+++ b/Engine/Source/Runtime/StorageServerClient/Private/DebugStorageServerConnection.cpp
@@ -108,7 +108,7 @@ void UDebugStorageServerConnection::Draw(UCanvas* Canvas, APlayerController*)
 		}
 
 		FString ZenConnectionDebugMsg;
-		ZenConnectionDebugMsg = FString::Printf(TEXT("ZenServer streaming from %s [%.2fMbps]"), *Owner->GetConnectionBackend()->GetHostName(), Throughput);
+		ZenConnectionDebugMsg = FString::Printf(TEXT("ZenServer streaming from %s [%.2fMbps]"), *HostAddress, Throughput);
 		GEngine->AddOnScreenDebugMessage((uint64)this, 86400.0f, FColor::White, ZenConnectionDebugMsg, false);
 		
 		History.push_back({ StatsTimeNow, MaxReqThroughput, MinReqThroughput, Throughput, ReqCount });
diff --git a/Engine/Source/Runtime/StorageServerClient/Private/DebugStorageServerConnection.h b/Engine/Source/Runtime/StorageServerClient/Private/DebugStorageServerConnection.h
--- a/Engine/Source/Runtime/StorageServerClient/Private/DebugStorageServerConnection.h
+++ b/Engine/Source/Runtime/StorageServerClient/Private/DebugStorageServerConnection.h
@@ -8,8 +8,6 @@
 
 #include "DebugStorageServerConnection.generated.h"
 
-class FStorageServerConnection;
-
 UCLASS()
 class UDebugStorageServerConnection : public UObject
 {
@@ -21,9 +19,9 @@ public:
 
 	void AddTimingInstance(double duration, uint64 bytes);
 
-	void SetOwner(FStorageServerConnection* Connection)
+	void SetHostAddress(FString Address)
 	{
-		Owner = Connection;
+		HostAddress = Address;
 	}
 
 	static void ShowGraph(FOutputDevice&);
@@ -54,7 +52,7 @@ private:
 	double MaxRequestThroughput = 0.0;
 
 	FCriticalSection StatsCS;
-	FStorageServerConnection* Owner;
+	FString HostAddress;
 
 	static bool ShowGraphs;
 
diff --git a/Engine/Source/Runtime/StorageServerClient/Private/StoragePlatformConnectionBackend.cpp b/Engine/Source/Runtime/StorageServerClient/Private/StoragePlatformConnectionBackend.cpp
deleted file mode 100644
diff --git a/Engine/Source/Runtime/StorageServerClient/Private/StoragePlatformConnectionBackend.h b/Engine/Source/Runtime/StorageServerClient/Private/StoragePlatformConnectionBackend.h
deleted file mode 100644
diff --git a/Engine/Source/Runtime/StorageServerClient/Private/StorageServerConnection.cpp b/Engine/Source/Runtime/StorageServerClient/Private/StorageServerConnection.cpp
--- a/Engine/Source/Runtime/StorageServerClient/Private/StorageServerConnection.cpp
+++ b/Engine/Source/Runtime/StorageServerClient/Private/StorageServerConnection.cpp
@@ -9,689 +9,642 @@
 #include "Misc/ScopeLock.h"
 #include "Misc/OutputDeviceRedirector.h"
 #include "Misc/StringBuilder.h"
+#include "Serialization/MemoryReader.h"
 #include "Serialization/CompactBinary.h"
 #include "Serialization/CompactBinarySerialization.h"
 #include "Memory/MemoryView.h"
 #include "Memory/SharedBuffer.h"
 #include "ProfilingDebugging/CountersTrace.h"
 #include "SocketSubsystem.h"
-#include "Sockets.h"
 #include "IO/PackageStore.h"
-#include "StorageSocketConnectionBackend.h"
-#include "StoragePlatformConnectionBackend.h"
-#include "GenericPlatform/GenericPlatformHostCommunication.h"
+#include "BuiltInHttpClient/BuiltInHttpClient.h"
+#include "BuiltInHttpClient/BuiltInHttpClientFSocket.h"
+#include "BuiltInHttpClient/BuiltInHttpClientPlatformSocket.h"
 #include "Engine/Engine.h"
 
 #if !UE_BUILD_SHIPPING
 
-DEFINE_LOG_CATEGORY_STATIC(LogStorageServerConnection, Log, All);
+DEFINE_LOG_CATEGORY(LogStorageServerConnection);
 
 TRACE_DECLARE_MEMORY_COUNTER(ZenHttpClientSerializedBytes, TEXT("ZenClient/SerializedBytes"));
 
-static uint64 GetCompressedOffset(const FCompressedBuffer& Buffer, uint64 RawOffset)
+FStorageServerConnection::~FStorageServerConnection()
 {
-	if (RawOffset > 0)
+	if (StatsObject)
 	{
-		uint64 BlockSize = 0;
-		ECompressedBufferCompressor Compressor;
-		ECompressedBufferCompressionLevel CompressionLevel;
-		const bool bOk = Buffer.TryGetCompressParameters(Compressor, CompressionLevel, BlockSize);
-		check(bOk);
-
-		return BlockSize > 0 ? RawOffset % BlockSize : 0;
+		StatsObject->StopDrawing();
+		StatsObject->RemoveFromRoot();
 	}
-
-	return 0;
 }
 
-FStorageServerRequest::FStorageServerRequest(FAnsiStringView Verb, FAnsiStringView Resource, FAnsiStringView Hostname, EStorageServerContentType Accept)
-: AcceptType(Accept)
+bool FStorageServerConnection::Initialize(TArrayView<const FString> HostAddresses, const int32 Port, const FAnsiStringView& InBaseURI)
 {
-	SetIsSaving(true);
-	HeaderBuffer << Verb << " " << Resource << " HTTP/1.1\r\n" 
-		<< "Host: " << Hostname << "\r\n"
-		<< "Connection: Keep-Alive\r\n"
-		<< "Accept: " << GetMimeTypeString(Accept) << "\r\n";
-}
+	BaseURI = InBaseURI;
+	TArray<FString> SortedHostAddresses = SortHostAddressesByLocalSubnet(HostAddresses, Port);
 
-IStorageConnectionSocket* FStorageServerRequest::Send(FStorageServerConnection& Owner, bool bLogOnError)
-{
-	StartTiming();
-	if (BodyBuffer.Num())
+	for (const FString& HostAddress : SortedHostAddresses)
 	{
-		HeaderBuffer.Append("Content-Length: ").Appendf("%d\r\n", BodyBuffer.Num());
-	}
-	HeaderBuffer << "\r\n";
-	int32 BytesLeft = HeaderBuffer.Len();
-	
-	auto Send = [](IStorageConnectionSocket* Socket, const uint8* Data, int32 Length)
-	{
-		int32 BytesLeft = Length;
-		while (BytesLeft > 0)
+		HttpClient = CreateHttpClient(HostAddress, Port);
+		CurrentHostAddr = HostAddress;
+		if (HandshakeRequest())
 		{
-			if (!Socket->Send(Data, BytesLeft))
+			if (StatsObject)
 			{
-				return false;
+				StatsObject->SetHostAddress(CurrentHostAddr);
 			}
-			BytesLeft -= BytesLeft;
+			return true;
 		}
-		return true;
-	};
-
-	int32 Attempts = 0;
-	FStorageConnectionBackend* Backend = Owner.GetConnectionBackend();
-	while (Attempts < 10)
-	{
-		IStorageConnectionSocket* Socket;
-		IStorageConnectionSocket* SocketFromPool;
-		Socket = SocketFromPool = Backend->AcquireSocketFromPool();
-		if (!Socket)
-		{
-			Socket = Backend->AcquireNewSocket();
-			if (!Socket)
-			{
-				++Attempts;
-				continue;
-			}
-		}
-
-		if (!Send(Socket, reinterpret_cast<const uint8*>(HeaderBuffer.GetData()), HeaderBuffer.Len()) ||
-			!Send(Socket, BodyBuffer.GetData(), BodyBuffer.Num()))
-		{
-			Backend->ReleaseSocket(Socket, false);
-			if (!SocketFromPool)
-			{
-				++Attempts;
-			}
-			continue;
-		}
-		return Socket;
 	}
 
-	if (bLogOnError)
-	{
-		UE_LOG(LogStorageServerConnection, Fatal, TEXT("Failed sending request to storage server."));
-	}
-	return nullptr;
+	HttpClient.Reset();
+	return false;
 }
 
-void FStorageServerRequest::Serialize(void* V, int64 Length)
+void FStorageServerConnection::PackageStoreRequest(TFunctionRef<void(FPackageStoreEntryResource&&)> Callback)
 {
-	int32 Index = BodyBuffer.AddUninitialized(Length);
-	uint8* Dest = BodyBuffer.GetData() + Index;
-	FMemory::Memcpy(Dest, V, Length);
-}
+	TAnsiStringBuilder<256> ResourceBuilder;
+	ResourceBuilder.Append(BaseURI).Append("/entries?fieldfilter=packagestoreentry");
 
-FStorageServerResponse::FStorageServerResponse(FStorageServerConnection& InOwner, IStorageConnectionSocket& InSocket)
-	: Owner(InOwner)
-	, Socket(&InSocket)
-{
-	SetIsLoading(true);
-	uint8 Buffer[1024];
-	int32 TotalReadFromSocket = 0;
-	auto ReadResponseLine = [&Buffer, &InSocket, &TotalReadFromSocket]() -> FAnsiStringView
+	IStorageServerHttpClient::FResult ResultTuple = HttpClient->RequestSync(*ResourceBuilder, EStorageServerContentType::CbObject);
+	TIoStatusOr<FIoBuffer> Result = ResultTuple.Get<0>();
+	if (Result.IsOk())
 	{
-		for (;;)
+		FMemoryReaderView Reader(Result.ValueOrDie().GetView());
+		FCbObject ResponseObj = LoadCompactBinary(Reader).AsObject();
 		{
-			uint64 BytesRead;
-			InSocket.Recv(Buffer, 1024, BytesRead, ESocketReceiveFlags::Peek);
-			FAnsiStringView ResponseView(reinterpret_cast<const ANSICHAR*>(Buffer), BytesRead);
-			int32 LineEndIndex;
-			if (ResponseView.FindChar('\r', LineEndIndex) && BytesRead >= LineEndIndex + 2)
+			TRACE_CPUPROFILER_EVENT_SCOPE(StorageServerPackageStoreRequestParseEntries);
+			for (FCbField& OplogEntry : ResponseObj["entries"].AsArray())
 			{
-				check(ResponseView[LineEndIndex + 1] == '\n');
-				InSocket.Recv(Buffer, LineEndIndex + 2, BytesRead, ESocketReceiveFlags::None);
-				check(BytesRead == LineEndIndex + 2);
-				TotalReadFromSocket += BytesRead;
-				return ResponseView.Left(LineEndIndex);
+				FCbObject OplogObj = OplogEntry.AsObject();
+				FPackageStoreEntryResource Entry = FPackageStoreEntryResource::FromCbObject(OplogObj["packagestoreentry"].AsObject());
+				Callback(MoveTemp(Entry));
 			}
 		}
-	};
-
-	FAnsiStringView ResponseLine = ReadResponseLine();
-	if (ResponseLine == "HTTP/1.1 200 OK")
-	{
-		bIsOk = true;
 	}
-	else if (ResponseLine.StartsWith("HTTP/1.1 "))
+	else
 	{
-		ErrorCode = TCString<ANSICHAR>::Atoi64(ResponseLine.GetData() + 9);
+		UE_LOG(LogStorageServerConnection, Fatal, TEXT("Failed to read oplog from storage server. '%s'"), *Result.Status().ToString());
 	}
-	while (!ResponseLine.IsEmpty())
+}
+
+void FStorageServerConnection::FileManifestRequest(TFunctionRef<void(FIoChunkId Id, FStringView Path, int64 RawSize)> Callback)
+{
+	TRACE_CPUPROFILER_EVENT_SCOPE(FStorageServerConnection::FileManifestRequest);
+
+	TAnsiStringBuilder<256> ResourceBuilder;
+	ResourceBuilder.Append(BaseURI).Append("/files?fieldnames=id,clientpath,rawsize");
+
+	IStorageServerHttpClient::FResult ResultTuple = HttpClient->RequestSync(*ResourceBuilder, EStorageServerContentType::CbObject);
+	TIoStatusOr<FIoBuffer> Result = ResultTuple.Get<0>();
+	if (Result.IsOk())
 	{
-		ResponseLine = ReadResponseLine();
-		if (ResponseLine.StartsWith("Content-Length: "))
-		{
-			ContentLength = FMath::Max(0, TCString<ANSICHAR>::Atoi64(ResponseLine.GetData() + 16));
-		}
-		else if (ResponseLine.StartsWith("Content-Type: "))
+		FMemoryReaderView Reader(Result.ValueOrDie().GetView());
+		FCbObject ResponseObj = LoadCompactBinary(Reader).AsObject();
+		for (FCbField& FileArrayEntry : ResponseObj["files"].AsArray())
 		{
-			ContentType = GetMimeType(ResponseLine.RightChop(14));
-		}
-	}
-	if (!bIsOk && ContentLength)
-	{
-		TArray<uint8> ErrorBuffer;
-		ErrorBuffer.SetNumUninitialized(ContentLength + 1);
-		uint64 BytesRead = 0;
-		InSocket.Recv(ErrorBuffer.GetData(), ContentLength, BytesRead, ESocketReceiveFlags::WaitAll);
+			FCbObject Entry = FileArrayEntry.AsObject();
+			FCbObjectId Id = Entry["id"].AsObjectId();
+			int64 ResponseRawSize = Entry["rawsize"].AsInt64(-1);
 
-		if (BytesRead > 0)
-		{
-			ErrorBuffer[BytesRead] = '\0';
-			ErrorMessage = FString(BytesRead, ANSI_TO_TCHAR(reinterpret_cast<ANSICHAR*>(ErrorBuffer.GetData())));
-		}
-		else
-		{
-			ErrorMessage = TEXT("Unknown error");
-		}
+			TStringBuilder<128> WidePath;
+			WidePath.Append(FUTF8ToTCHAR(Entry["clientpath"].AsString()));
 
-		ContentLength = 0;
+			FIoChunkId ChunkId;
+			ChunkId.Set(Id.GetView());
+
+			Callback(ChunkId, WidePath, ResponseRawSize);
+		}
 	}
-	if (ContentLength == 0)
+	else
 	{
-		ReleaseSocket(true);
+		UE_LOG(LogStorageServerConnection, Fatal, TEXT("Failed to read file manifest from storage server. '%s'"), *Result.Status().ToString());
 	}
 }
 
-FStorageServerChunkBatchRequest::FStorageServerChunkBatchRequest(FStorageServerConnection& InOwner, FAnsiStringView Resource, FAnsiStringView Hostname)
-	: FStorageServerRequest("POST", Resource, Hostname)
-	, Owner(InOwner)
+int64 FStorageServerConnection::ChunkSizeRequest(const FIoChunkId& ChunkId)
 {
-	uint32 Magic = 0xAAAA'77AC;
-	uint32 ChunkCountPlaceHolder = 0;
-	uint32 Reserved1 = 0;
-	uint32 Reserved2 = 0;
-	*this << Magic;
-	ChunkCountOffset = BodyBuffer.Num();
-	*this << ChunkCountPlaceHolder << Reserved1 << Reserved2;
-}
+	TRACE_CPUPROFILER_EVENT_SCOPE(FStorageServerConnection::ChunkSizeRequest);
 
-FStorageServerChunkBatchRequest& FStorageServerChunkBatchRequest::AddChunk(const FIoChunkId& ChunkId, int64 Offset, int64 Size)
-{
-	uint32* ChunkCount = reinterpret_cast<uint32*>(BodyBuffer.GetData() + ChunkCountOffset);
-	*this << const_cast<FIoChunkId&>(ChunkId) << *ChunkCount << Offset << Size;
-	++(*ChunkCount);
-	return *this;
-}
+	TAnsiStringBuilder<256> ResourceBuilder;
+	ResourceBuilder.Append(BaseURI);
+	ResourceBuilder << "/" << ChunkId << "/info";
 
-bool FStorageServerChunkBatchRequest::Issue(TFunctionRef<void(uint32 ChunkCount, uint32* ChunkIndices, uint64* ChunkSizes, FStorageServerResponse& ChunkDataStream)> OnResponse)
-{
-	IStorageConnectionSocket* Socket = Send(Owner);
-	if (!Socket)
+	const double StartTime = FPlatformTime::Seconds();
+	IStorageServerHttpClient::FResult ResultTuple = HttpClient->RequestSync(*ResourceBuilder, EStorageServerContentType::CbObject);
+	TIoStatusOr<FIoBuffer> Result = ResultTuple.Get<0>();
+	if (Result.IsOk())
 	{
-		UE_LOG(LogStorageServerConnection, Fatal, TEXT("Failed to send chunk batch request to storage server."));
-		return false;
-	}
-	FStorageServerResponse Response(Owner, *Socket);
-	if (!Response.IsOk())
-	{
-		UE_LOG(LogStorageServerConnection, Fatal, TEXT("Failed to read chunk batch from storage server. '%s'"), *Response.GetErrorMessage());
-		return false;
-	}
+		const double Duration = FPlatformTime::Seconds() - StartTime;
+		AddTimingInstance(Duration, Result.ValueOrDie().GetSize());
 
-	uint32 Magic;
-	uint32 ChunkCount;
-	uint32 Reserved1;
-	uint32 Reserved2;
-	Response << Magic;
-	if (Magic != 0xbada'b00f)
-	{
-		UE_LOG(LogStorageServerConnection, Fatal, TEXT("Invalid magic in chunk batch response from storage server."));
-		return false;
+		FMemoryReaderView Reader(Result.ValueOrDie().GetView());
+		FCbObject ResponseObj = LoadCompactBinary(Reader).AsObject();
+		const int64 ChunkSize = ResponseObj["size"].AsInt64(0);
+		return ChunkSize;
 	}
-	Response << ChunkCount;
-	if (ChunkCount > INT32_MAX)
+	else if (Result.Status().GetErrorCode() != EIoErrorCode::NotFound)
 	{
-		UE_LOG(LogStorageServerConnection, Fatal, TEXT("Invalid chunk count in chunk batch response from storage server."));
-		return false;
+		UE_LOG(LogStorageServerConnection, Fatal, TEXT("Failed to get chunk size from storage server. '%s'"), *Result.Status().ToString());
 	}
-	Response << Reserved1;
-	Response << Reserved2;
 
-	TArray<uint32, TInlineAllocator<64>> ChunkIndices;
-	ChunkIndices.Reserve(ChunkCount);
-	TArray<uint64, TInlineAllocator<64>> ChunkSizes;
-	ChunkSizes.Reserve(ChunkCount);
-	for (uint32 Index = 0; Index < ChunkCount; ++Index)
-	{
-		uint32 ChunkIndex;
-		uint32 Flags;
-		int64 ChunkSize;
-		Response << ChunkIndex;
-		Response << Flags;
-		Response << ChunkSize;
-		ChunkIndices.Add(ChunkIndex);
-		ChunkSizes.Emplace(ChunkSize);
-	}
-	OnResponse(ChunkCount, ChunkIndices.GetData(), ChunkSizes.GetData(), Response);
-	Owner.AddTimingInstance(GetDuration(), (double)Response.Tell());
-	return true;
+	return -1;
 }
 
-void FStorageServerResponse::ReleaseSocket(bool bKeepAlive)
+TIoStatusOr<FIoBuffer> FStorageServerConnection::ReadChunkRequest(
+	const FIoChunkId& ChunkId,
+	const uint64 Offset,
+	const uint64 Size,
+	const TOptional<FIoBuffer> OptDestination,
+	const bool bHardwareTargetBuffer
+)
 {
-	Owner.GetConnectionBackend()->ReleaseSocket(Socket, bKeepAlive);
-	Socket = nullptr;
-}
+	TRACE_CPUPROFILER_EVENT_SCOPE(FStorageServerConnection::ReadChunkRequest);
 
-void FStorageServerResponse::Serialize(void* V, int64 Length)
-{
-	if (Length == 0)
-	{
-		return;
-	}
-	if (!Socket)
-	{
-		UE_LOG(LogStorageServerConnection, Fatal, TEXT("Trying to read %lld bytes from released socket"), Length);
-		return;
-	}
-	if (Position + Length > ContentLength)
-	{
-		UE_LOG(LogStorageServerConnection, Fatal, TEXT("Trying to read %lld bytes from socket with only %lld available"), Length, ContentLength - Position);
-		return;
-	}
-	uint64 RemainingBytesToRead = Length;
-	uint8* Destination = reinterpret_cast<uint8*>(V);
-	while (RemainingBytesToRead)
-	{
-		uint64 BytesToRead32 = FMath::Min(RemainingBytesToRead, static_cast<uint64>(INT32_MAX));
-		uint64 BytesRead;
-		if (!Socket->Recv(Destination, static_cast<int32>(BytesToRead32), BytesRead, ESocketReceiveFlags::WaitAll))
-		{
-			UE_LOG(LogStorageServerConnection, Fatal, TEXT("Failed reading %d bytes from socket"), BytesToRead32);
-			return;
-		}
-		RemainingBytesToRead -= BytesRead;
-		Destination += BytesRead;
-		Position += BytesRead;
-	}
+	const double StartTime = FPlatformTime::Seconds();
 
-	TRACE_COUNTER_ADD(ZenHttpClientSerializedBytes, Length);
-	
-	if (Position == ContentLength)
+	TAnsiStringBuilder<256> ResourceBuilder;
+	BuildReadChunkRequestUrl(ResourceBuilder, ChunkId, Offset, Size);
+
+	IStorageServerHttpClient::FResult HttpResultTuple = HttpClient->RequestSync(*ResourceBuilder);
+	TIoStatusOr<FIoBuffer> ResultBuffer = ReadChunkRequestProcessHttpResult(HttpResultTuple, Offset, Size, OptDestination, bHardwareTargetBuffer);
+
+	if (ResultBuffer.IsOk())
 	{
-		ReleaseSocket(true);
+		const double Duration = FPlatformTime::Seconds() - StartTime;
+		AddTimingInstance(Duration, ResultBuffer.ValueOrDie().GetSize());
 	}
+
+	return ResultBuffer;
 }
 
-int64 FStorageServerResponse::SerializeChunk(FStorageServerSerializationContext& Context, FIoBuffer& OutChunk, void* TargetVa, uint64 RawOffset, uint64 RawSize, bool bHardwareTargetBuffer)
+void FStorageServerConnection::ReadChunkRequestAsync(
+	const FIoChunkId& ChunkId,
+	const uint64 Offset,
+	const uint64 Size,
+	const TOptional<FIoBuffer> OptDestination,
+	const bool bHardwareTargetBuffer,
+	TFunctionRef<void(TIoStatusOr<FIoBuffer> Data)> OnResponse
+)
 {
-	if (ContentLength == 0)
-	{
-		return 0;
-	}
+	TRACE_CPUPROFILER_EVENT_SCOPE(FStorageServerConnection::ReadChunkRequestAsync);
 
-	if (ContentType == EStorageServerContentType::Binary)
-	{
-		const uint64 ChunkSize = FMath::Min<uint64>(ContentLength, RawSize);
-		OutChunk = TargetVa ? FIoBuffer(FIoBuffer::Wrap, TargetVa, ChunkSize) : FIoBuffer(ChunkSize);
-		Serialize(OutChunk.Data(), OutChunk.DataSize());
-		return ChunkSize;
-	}
-	else if (ContentType == EStorageServerContentType::CompressedBinary)
-	{
-		TRACE_CPUPROFILER_EVENT_SCOPE(ZenClient::SerializeCompressedChunk);
+	const double StartTime = FPlatformTime::Seconds();
 
-		Context.CompressedBuffer.Reset(ContentLength);
-		Serialize(Context.CompressedBuffer.GetData(), ContentLength);
+	TAnsiStringBuilder<256> ResourceBuilder;
+	BuildReadChunkRequestUrl(ResourceBuilder, ChunkId, Offset, Size);
 
-		if (FCompressedBuffer Compressed = FCompressedBuffer::FromCompressed(FSharedBuffer::MakeView(Context.CompressedBuffer.GetData(), ContentLength)))
-		{
-			FCompressedBufferReaderSourceScope Source(Context.Decoder, Compressed);
-			const uint64 ChunkSize = FMath::Min(Compressed.GetRawSize(), RawSize);
-			OutChunk = TargetVa ? FIoBuffer(FIoBuffer::Wrap, TargetVa, ChunkSize) : FIoBuffer(ChunkSize);
-			const uint64 CompressedOffset = GetCompressedOffset(Compressed, RawOffset);
-			if (Context.Decoder.TryDecompressTo(OutChunk.GetMutableView(), CompressedOffset, bHardwareTargetBuffer ? ECompressedBufferDecompressFlags::IntermediateBuffer :ECompressedBufferDecompressFlags::None ))
-			{
-				return ChunkSize;
-			}
+	HttpClient->RequestAsync([this, Offset, Size, OptDestination, bHardwareTargetBuffer, OnResponse, StartTime](IStorageServerHttpClient::FResult HttpResultTuple)
+	{
+		TIoStatusOr<FIoBuffer> ResultBuffer = ReadChunkRequestProcessHttpResult(HttpResultTuple, Offset, Size, OptDestination, bHardwareTargetBuffer);
 
-			return 0;
+		if (ResultBuffer.IsOk())
+		{
+			const double Duration = FPlatformTime::Seconds() - StartTime;
+			AddTimingInstance(Duration, ResultBuffer.ValueOrDie().GetSize());
 		}
-	}
 
-	UE_LOG(LogStorageServerConnection, Fatal, TEXT("Received unknown chunk type from storage server"));
-	return 0;
+		OnResponse(ResultBuffer);
+	}, *ResourceBuilder);
 }
 
-int64 FStorageServerResponse::SerializeChunkTo(FMutableMemoryView Memory, uint64 RawOffset)
+TArray<FString> FStorageServerConnection::SortHostAddressesByLocalSubnet(TArrayView<const FString> HostAddresses, const int32 Port)
 {
-	if (ContentLength == 0)
+	bool bAllArePlatformSocketAddresses = true;
+	for (const FString& HostAddress : HostAddresses)
 	{
-		return 0;
+		if (!IsPlatformSocketAddress(HostAddress))
+		{
+			bAllArePlatformSocketAddresses = false;
+			break;
+		}
 	}
 
-	if (ContentType == EStorageServerContentType::Binary)
+	// return array without sorting if it's 0 or 1 addresses or all of them are platform sockets
+	if (HostAddresses.Num() <= 1 || bAllArePlatformSocketAddresses)
 	{
-		FMutableMemoryView Dst = Memory.Left(FMath::Min<uint64>(Memory.GetSize(), ContentLength));
-		Serialize(Dst.GetData(), Dst.GetSize());
-		return Dst.GetSize();
+		return TArray<FString>(HostAddresses);
 	}
-	else if (ContentType == EStorageServerContentType::CompressedBinary)
-	{
-		TRACE_CPUPROFILER_EVENT_SCOPE(ZenClient::SerializeCompressedChunk);
-		
-		TArray64<uint8> CompressedBuffer;
-		CompressedBuffer.Reset(ContentLength);
-		Serialize(CompressedBuffer.GetData(), ContentLength);
-		
-		if (FCompressedBuffer Compressed = FCompressedBuffer::FromCompressed(FSharedBuffer::MakeView(CompressedBuffer.GetData(), ContentLength)))
-		{
-			const uint64 CompressedOffset = GetCompressedOffset(Compressed, RawOffset);
-			FMutableMemoryView Dst = Memory.Left(FMath::Min(Memory.GetSize(), Compressed.GetRawSize() - CompressedOffset));
-			if (FCompressedBufferReader(Compressed).TryDecompressTo(Dst, CompressedOffset))
-			{
-				return Dst.GetSize();
-			}
 
-			return 0;
-		}
-	}
-	
-	UE_LOG(LogStorageServerConnection, Fatal, TEXT("Received unknown chunk type from storage server"));
-	return 0;
-}
+	TArray<FString> Result;
+	ISocketSubsystem& SocketSubsystem = *ISocketSubsystem::Get();
 
-FCbObject FStorageServerResponse::GetResponseObject()
-{
-	FCbField Payload;
+	// Sorting logic in order:
+	// - special platform socket address, see PlatformSocketAddress
+	// - on desktop, if it's an IPV6 address loopback (ends with ":1")
+	// - on desktop, if it's and IPV4 address loopback (starts with "127.0.0")
+	// - host IPV4 subnet matches the client subnet (xxx.xxx.xxx)
+	// - remaining addresses
+	bool bCanBindAll = false;
+	bool bAppendPort = false;
+	TSharedPtr<FInternetAddr> localAddr = SocketSubsystem.GetLocalHostAddr(*GLog, bCanBindAll);
+	FString localAddrStringSubnet = localAddr->ToString(bAppendPort);
 
-	if (ContentType == EStorageServerContentType::CompressedBinary)
+	int32 localLastDotPos = INDEX_NONE;
+	if (localAddrStringSubnet.FindLastChar(TEXT('.'), localLastDotPos))
 	{
-		TRACE_CPUPROFILER_EVENT_SCOPE(GetCompressedResponseObject);
+		localAddrStringSubnet = localAddrStringSubnet.LeftChop(localAddrStringSubnet.Len() - localLastDotPos);
+	}
 
-		TArray64<uint8> CompressedBuffer;
-		CompressedBuffer.Reset(ContentLength);
+	TArray<FString> PlatformSocketAddress;
+	TArray<FString> IPV6Loopback;
+	TArray<FString> IPV4Loopback;
+	TArray<FString> RegularAddresses;
+
+	for (const FString& HostAddress : HostAddresses)
+	{
+		if (IsPlatformSocketAddress(HostAddress))
 		{
-			TRACE_CPUPROFILER_EVENT_SCOPE(SerializeCompressedData);
-			Serialize(CompressedBuffer.GetData(), ContentLength);
+			PlatformSocketAddress.Push(HostAddress);
+			continue;
 		}
 
-		if (FCompressedBuffer Compressed = FCompressedBuffer::FromCompressed(FSharedBuffer::MakeView(CompressedBuffer.GetData(), ContentLength)))
+		TSharedPtr<FInternetAddr> Addr = StringToInternetAddr(HostAddress, Port);
+		FString tempAddrStringSubnet = Addr->ToString(bAppendPort);
+
+#if PLATFORM_DESKTOP || PLATFORM_ANDROID
+		if (Addr->GetProtocolType() == FNetworkProtocolTypes::IPv6)
+		{
+			if (tempAddrStringSubnet.EndsWith(":1"))
+			{
+				IPV6Loopback.Push(HostAddress);
+				continue;
+			}
+		}
+		else
 		{
-			FIoBuffer Decompressed(Compressed.GetRawSize());
-			if (FCompressedBufferReader(Compressed).TryDecompressTo(Decompressed.GetMutableView(), 0))
+			if (tempAddrStringSubnet.StartsWith("127.0.0."))
 			{
-				FBufferReader DecompressedAr(Decompressed.GetData(), Decompressed.GetSize(), false);
-				{
-					TRACE_CPUPROFILER_EVENT_SCOPE(LoadingUncompressedCompactBinary);
-					Payload = LoadCompactBinary(DecompressedAr);
-				}
+				IPV4Loopback.Push(HostAddress);
+				continue;
 			}
 		}
+#endif
+
+		int32 LastDotPos = INDEX_NONE;
+		if (tempAddrStringSubnet.FindLastChar(TEXT('.'), LastDotPos))
+		{
+			tempAddrStringSubnet = tempAddrStringSubnet.LeftChop(tempAddrStringSubnet.Len() - LastDotPos);
+		}
+
+		if (localAddrStringSubnet.Equals(tempAddrStringSubnet))
+			RegularAddresses.Insert(HostAddress, 0);
+		else
+			RegularAddresses.Push(HostAddress);
 	}
-	else if (ContentType == EStorageServerContentType::CbObject)
+
+	for (const FString& PlatformAddr : PlatformSocketAddress)
 	{
-		Payload = LoadCompactBinary(*this);
+		Result.Push(PlatformAddr);
 	}
-	
-	return Payload.AsObject();
-}
-
-FStorageConnectionBackend::FStorageConnectionBackend(FStorageServerConnection& InOwner)
-	: Owner(InOwner)
-{
-}
 
-bool FStorageConnectionBackend::Initialize(TArrayView<const FString> InHostAddresses, int32 InPort, const FAnsiStringView& InBaseURI)
-{
-	BaseURI = InBaseURI;
-	return InitializeInternal(InHostAddresses, InPort);
-}
+	for (const FString& Addrv6lb : IPV6Loopback)
+	{
+		Result.Push(Addrv6lb);
+	}
 
-FStorageServerConnection::FStorageServerConnection()
-	: SocketSubsystem(*ISocketSubsystem::Get())
-{
-}
+	for (const FString& Addrv4lb : IPV4Loopback)
+	{
+		Result.Push(Addrv4lb);
+	}
 
-FStorageServerConnection::~FStorageServerConnection()
-{
-	if (StatsObject)
+	for (const FString& RegularAddr : RegularAddresses)
 	{
-		StatsObject->StopDrawing();
-		StatsObject->RemoveFromRoot();
+		Result.Push(RegularAddr);
 	}
+
+	return Result;
 }
 
-bool FStorageServerConnection::Initialize(TArrayView<const FString> InHostAddresses, int32 InPort, const FAnsiStringView& InBaseURI)
+bool FStorageServerConnection::IsPlatformSocketAddress(const FString Address)
 {
-	BaseURI = InBaseURI;
-	if (!CreateConnectionBackend(InHostAddresses, InPort))
-	{
-		return false;
-	}
-	return ConnectionBackend->Initialize(InHostAddresses, InPort, BaseURI);
+	return Address.StartsWith(TEXT("platform://"));
 }
 
-void FStorageServerConnection::PackageStoreRequest(TFunctionRef<void(FPackageStoreEntryResource&&)> Callback)
+TUniquePtr<IStorageServerHttpClient> FStorageServerConnection::CreateHttpClient(const FString Address, const int32 Port)
 {
-	TRACE_CPUPROFILER_EVENT_SCOPE(StorageServerPackageStoreRequest);
+	TSharedPtr<FInternetAddr> Addr = StringToInternetAddr(Address, Port);
 
-	TAnsiStringBuilder<256> ResourceBuilder;
-	ResourceBuilder.Append(BaseURI).Append("/entries?fieldfilter=packagestoreentry");
-	FStorageServerRequest Request("GET", *ResourceBuilder, Hostname, EStorageServerContentType::CompressedBinary);
-	IStorageConnectionSocket* Socket = Request.Send(*this);
-	if (!Socket)
-	{
-		UE_LOG(LogStorageServerConnection, Fatal, TEXT("Failed to send oplog request to storage server at %s."), *ServerAddr->ToString(true));
-		return;
-	}
-	FStorageServerResponse Response(*this, *Socket);
-	if (Response.IsOk())
-	{
-		FCbObject ResponseObj = Response.GetResponseObject();
+	// Use Address as Hostname if we can't resolve FInternetAddr
+	FString HostName = Addr.IsValid() ? Addr->ToString(false) : Address;
 
-		TRACE_CPUPROFILER_EVENT_SCOPE(StorageServerPackageStoreRequestParseEntries);
-		for (FCbField& OplogEntry : ResponseObj["entries"].AsArray())
-		{
-			FCbObject OplogObj = OplogEntry.AsObject();
-			FPackageStoreEntryResource Entry = FPackageStoreEntryResource::FromCbObject(OplogObj["packagestoreentry"].AsObject());
-			Callback(MoveTemp(Entry));
-		}
+	UE_LOG(LogStorageServerConnection, Display, TEXT("Creating zen store connection to %s:%i (\"%s\")."), *Address, Port, *HostName);
+
+	TUniquePtr<IBuiltInHttpClientSocketPool> SocketPool;
+	if (IsPlatformSocketAddress(Address))
+	{
+		SocketPool = MakeUnique<FBuiltInHttpClientPlatformSocketPool>(Address); 
 	}
 	else
 	{
-		UE_LOG(LogStorageServerConnection, Fatal, TEXT("Failed to read oplog from storage server at %s. '%s'"), *ServerAddr->ToString(true), *Response.GetErrorMessage());
+		SocketPool = MakeUnique<FBuiltInHttpClientFSocketPool>(Addr, *ISocketSubsystem::Get());
 	}
+
+	return MakeUnique<FBuiltInHttpClient>(MoveTemp(SocketPool), HostName);
 }
 
-void FStorageServerConnection::FileManifestRequest(TFunctionRef<void(FIoChunkId Id, FStringView Path, int64 RawSize)> Callback)
+TSharedPtr<FInternetAddr> FStorageServerConnection::StringToInternetAddr(const FString HostAddr, const int32 Port)
 {
-	TAnsiStringBuilder<256> ResourceBuilder;
-	ResourceBuilder.Append(BaseURI).Append("/files?fieldnames=id,clientpath,rawsize");
-	FStorageServerRequest Request("GET", *ResourceBuilder, Hostname, EStorageServerContentType::CbObject);
-	IStorageConnectionSocket* Socket = Request.Send(*this);
-	if (!Socket)
+	TSharedPtr<FInternetAddr> Result = TSharedPtr<FInternetAddr>();
+
+	if (IsPlatformSocketAddress(HostAddr))
 	{
-		UE_LOG(LogStorageServerConnection, Fatal, TEXT("Failed to send file manifest request to storage server at %s."), *ServerAddr->ToString(true));
-		return;
+		return Result;
 	}
 
-	FStorageServerResponse Response(*this, *Socket);
-	if (Response.IsOk())
+	ISocketSubsystem& SocketSubsystem = *ISocketSubsystem::Get();
+
+	// Numeric IPV6 addresses can be enclosed in brackets, and must have the brackets stripped before calling GetAddressFromString
+	FString ModifiedHostAddr;
+	const FString* EffectiveHostAddr = &HostAddr;
+	if (!HostAddr.IsEmpty() && HostAddr[0] == TEXT('[') && HostAddr[HostAddr.Len() - 1] == TEXT(']'))
 	{
-		FCbObject ResponseObj = Response.GetResponseObject();
-		
-		for (FCbField& FileArrayEntry : ResponseObj["files"].AsArray())
+#if PLATFORM_HAS_BSD_SOCKETS && !PLATFORM_HAS_BSD_IPV6_SOCKETS
+		// If the platform doesn't have IPV6 BSD Sockets, then handle an attempt at conversion of loopback addresses, and skip and warn about other addresses
+		if (HostAddr == TEXT("[::1]"))
 		{
-			FCbObject Entry = FileArrayEntry.AsObject();
-			FCbObjectId Id = Entry["id"].AsObjectId();
-			int64 ResponseRawSize = Entry["rawsize"].AsInt64(-1);
-
-			TStringBuilder<128> WidePath;
-			WidePath.Append(FUTF8ToTCHAR(Entry["clientpath"].AsString()));
-
-			FIoChunkId ChunkId;
-			ChunkId.Set(Id.GetView());
-
+			// Substitute IPV4 loopback for IPV6 loopback
+			ModifiedHostAddr = TEXT("127.0.0.1");
+		}
+		else
+		{
+			UE_LOG(LogStorageServerConnection, Warning, TEXT("Ignoring storage server host IPV6 address on platform that doesn't support IPV6: %s"), *HostAddr);
+			return TSharedPtr<FInternetAddr>();
+		}
+#else
+		ModifiedHostAddr = FStringView(HostAddr).Mid(1, HostAddr.Len() - 2);
+#endif
+		EffectiveHostAddr = &ModifiedHostAddr;
+	}
 
-			Callback(ChunkId, WidePath, ResponseRawSize);
+	Result = SocketSubsystem.GetAddressFromString(*EffectiveHostAddr);
+	if (!Result.IsValid() || !Result->IsValid())
+	{
+		FAddressInfoResult GAIRequest = SocketSubsystem.GetAddressInfo(**EffectiveHostAddr, nullptr, EAddressInfoFlags::Default, NAME_None);
+		if (GAIRequest.ReturnCode == SE_NO_ERROR && GAIRequest.Results.Num() > 0)
+		{
+			Result = GAIRequest.Results[0].Address;
 		}
 	}
-	else
+
+	if (Result.IsValid() && Result->IsValid())
 	{
-		UE_LOG(LogStorageServerConnection, Fatal, TEXT("Failed to read file manifest from storage server at %s. '%s'"), *ServerAddr->ToString(true), *Response.GetErrorMessage());
+		Result->SetPort(Port);
 	}
+
+	return Result;
 }
 
-int64 FStorageServerConnection::ChunkSizeRequest(const FIoChunkId& ChunkId)
+bool FStorageServerConnection::HandshakeRequest()
 {
 	TAnsiStringBuilder<256> ResourceBuilder;
 	ResourceBuilder.Append(BaseURI);
-	ResourceBuilder << "/" << ChunkId << "/info";
 
-	FStorageServerRequest Request("GET", *ResourceBuilder, Hostname, EStorageServerContentType::CbObject);
-	IStorageConnectionSocket* Socket = Request.Send(*this);
-	if (!Socket)
-	{
-		UE_LOG(LogStorageServerConnection, Fatal, TEXT("Failed to send chunk size request to storage server at %s."), *ServerAddr->ToString(true));
-		return -1;
+	// Handshakes are done with a limited connection timeout so that we can find out if the destination is unreachable in a timely manner.
+	const float ConnectionTimeoutSeconds = 5.0f;
+
+	IStorageServerHttpClient::FResult ResultTuple = HttpClient->RequestSync(
+		*ResourceBuilder,
+		EStorageServerContentType::Unknown,
+		"GET",
+		TOptional<FIoBuffer>(),
+		TOptional<FIoBuffer>(),
+		ConnectionTimeoutSeconds,
+		false
+	);
+	TIoStatusOr<FIoBuffer> Result = ResultTuple.Get<0>();
+	if (Result.IsOk())
+	{
+		FMemoryReaderView Reader(Result.ValueOrDie().GetView());
+		FCbObject ResponseObj = LoadCompactBinary(Reader).AsObject();
+		// we currently don't have any concept of protocol versioning, if
+		// we succeed in communicating with the endpoint we're good since
+		// any breaking API change would need to be done in a backward
+		// compatible manner
+		return true;
 	}
-	FStorageServerResponse Response(*this, *Socket);
-	if (Response.IsOk())
-	{
-		AddTimingInstance(Request.GetDuration(), (double)Response.TotalSize());
-
-		FCbObject ResponseObj = Response.GetResponseObject();
 
-		const int64 ChunkSize = ResponseObj["size"].AsInt64(0);
-
-		return ChunkSize;
-	}
-	else if (Response.GetErrorCode() == 404)
-	{
-		return -1;
-	}
-	else
-	{
-		UE_LOG(LogStorageServerConnection, Fatal, TEXT("Failed to get chunk size from storage server at %s. '%s'"), *ServerAddr->ToString(true), *Response.GetErrorMessage());
-	}
-	return -1;
+	return false;
 }
 
-bool FStorageServerConnection::ReadChunkRequest(const FIoChunkId& ChunkId, uint64 Offset, uint64 Size, TFunctionRef<void(FStorageServerResponse&)> OnResponse)
+void FStorageServerConnection::BuildReadChunkRequestUrl(FAnsiStringBuilderBase& Builder, const FIoChunkId& ChunkId, const uint64 Offset, const uint64 Size)
 {
-	TRACE_CPUPROFILER_EVENT_SCOPE(ZenHttpClient::ReadChunkRequest);
-	TAnsiStringBuilder<256> ResourceBuilder;
-	ResourceBuilder.Append(BaseURI) << "/" << ChunkId;
-
+	Builder.Append(BaseURI) << "/" << ChunkId;
 	bool HaveQuery = false;
-
-	auto AppendQueryDelimiter = [&] 
+	auto AppendQueryDelimiter = [&]
 	{
 		if (HaveQuery)
 		{
-			ResourceBuilder.Append(ANSITEXTVIEW("&"));
+			Builder.Append(ANSITEXTVIEW("&"));
 		}
 		else
 		{
-			ResourceBuilder.Append(ANSITEXTVIEW("?"));
+			Builder.Append(ANSITEXTVIEW("?"));
 			HaveQuery = true;
 		}
 	};
-
 	if (Offset)
 	{
 		AppendQueryDelimiter();
-		ResourceBuilder.Appendf("offset=%" UINT64_FMT, Offset);
+		Builder.Appendf("offset=%" UINT64_FMT, Offset);
 	}
-
 	if (Size != ~uint64(0))
 	{
 		AppendQueryDelimiter();
-		ResourceBuilder.Appendf("size=%" UINT64_FMT, Size);
+		Builder.Appendf("size=%" UINT64_FMT, Size);
 	}
+}
 
-	FStorageServerRequest Request("GET", *ResourceBuilder, Hostname, EStorageServerContentType::CompressedBinary);
-	IStorageConnectionSocket* Socket = Request.Send(*this);
-	if (!Socket)
-	{
-		UE_LOG(LogStorageServerConnection, Fatal, TEXT("Failed to send chunk read request to storage server at %s."), *ServerAddr->ToString(true));
-		return false;
-	}
-	FStorageServerResponse Response(*this, *Socket);
-	if (Response.IsOk())
+TIoStatusOr<FIoBuffer> FStorageServerConnection::ReadChunkRequestProcessHttpResult(
+	IStorageServerHttpClient::FResult ResultTuple,
+	const uint64 Offset,
+	const uint64 Size,
+	const TOptional<FIoBuffer> OptDestination,
+	const bool bHardwareTargetBuffer
+)
+{
+	TIoStatusOr<FIoBuffer> Result = ResultTuple.Get<0>();
+	EStorageServerContentType MimeType = ResultTuple.Get<1>();
+	if (!Result.IsOk())
 	{
-		OnResponse(Response);
-		AddTimingInstance(Request.GetDuration(), (double)Response.Tell());
-		return true;
+		UE_LOG(LogStorageServerConnection, Fatal, TEXT("Failed to get chunk size from storage server. '%s'"), *Result.Status().ToString());
+		return Result.Status();
 	}
-	else if (Response.GetErrorCode() == 404)
+
+	FIoBuffer Buffer = Result.ValueOrDie();
+	TRACE_COUNTER_ADD(ZenHttpClientSerializedBytes, Buffer.GetSize());
+
+	if (MimeType == EStorageServerContentType::Binary)
 	{
-		return false;
+		TRACE_CPUPROFILER_EVENT_SCOPE(FStorageServerConnection::ReadChunkRequest::Binary);
+
+		if (OptDestination.IsSet())
+		{
+			ensure(OptDestination->GetSize() >= Buffer.GetSize());
+
+			FIoBuffer Destination = OptDestination.GetValue();
+			FMemory::Memcpy(Destination.GetData(), Buffer.GetData(), Buffer.GetSize());
+			Destination.SetSize(Buffer.GetSize());
+			return Destination;
+		}
+		else
+		{
+			return Buffer;
+		}
 	}
-	else
+	else if (MimeType == EStorageServerContentType::CompressedBinary)
 	{
-		UE_LOG(LogStorageServerConnection, Fatal, TEXT("Failed to read chunk from storage server at %s. '%s'"), *ServerAddr->ToString(true), *Response.GetErrorMessage());
-		return false;
+		TRACE_CPUPROFILER_EVENT_SCOPE(FStorageServerConnection::ReadChunkRequest::CompressedBinary);
+
+		FMemoryReaderView Reader(Buffer.GetView());
+		FCompressedBuffer CompressedBuffer = FCompressedBuffer::FromCompressed(FSharedBuffer::MakeView(Buffer.GetData(), Buffer.GetSize()));
+		FCompressedBufferReader CompressedBufferReader(CompressedBuffer);
+		const uint64 RawSize = CompressedBufferReader.GetRawSize();
+		if (RawSize > 0)
+		{
+			const uint64 CompressedOffset = GetCompressedOffset(CompressedBuffer, Offset);
+			const uint64 BytesToReadNonTrimmed = Size > 0 ? FMath::Min(Size, RawSize) : RawSize;
+			const uint64 BytesToRead = FMath::Min(BytesToReadNonTrimmed, RawSize - CompressedOffset);
+
+			ensure(!OptDestination.IsSet() || OptDestination->GetSize() >= BytesToRead);
+
+			FIoBuffer OutChunk = OptDestination.IsSet() ? OptDestination.GetValue() : FIoBuffer(BytesToRead);
+			OutChunk.SetSize(BytesToRead);
+
+			if (CompressedBufferReader.TryDecompressTo(OutChunk.GetMutableView(), CompressedOffset, bHardwareTargetBuffer ? ECompressedBufferDecompressFlags::IntermediateBuffer : ECompressedBufferDecompressFlags::None))
+			{
+				return OutChunk;
+			}
+		}
 	}
-}
 
-FStorageServerChunkBatchRequest FStorageServerConnection::NewChunkBatchRequest()
-{
-	TAnsiStringBuilder<256> ResourceBuilder;
-	ResourceBuilder.Append(BaseURI).Append("/batch");
-	return FStorageServerChunkBatchRequest(*this, *ResourceBuilder, Hostname);
+	return FIoStatus(EIoErrorCode::Unknown);
 }
 
-FString FStorageServerConnection::GetHostAddr() const
+uint64 FStorageServerConnection::GetCompressedOffset(const FCompressedBuffer& Buffer, uint64 RawOffset)
 {
-	return ServerAddr.IsValid() ? ServerAddr->ToString(false) : FString();
+	if (RawOffset > 0)
+	{
+		uint64 BlockSize = 0;
+		ECompressedBufferCompressor Compressor;
+		ECompressedBufferCompressionLevel CompressionLevel;
+		const bool bOk = Buffer.TryGetCompressParameters(Compressor, CompressionLevel, BlockSize);
+		check(bOk);
+
+		return BlockSize > 0 ? RawOffset % BlockSize : 0;
+	}
+
+	return 0;
 }
 
-bool FStorageServerConnection::CreateConnectionBackend(TArrayView<const FString> HostAddresses, int32 Port)
+void FStorageServerConnection::AddTimingInstance(const double Duration, const uint64 Bytes)
 {
-	FString PlatformAddress;
-	for (const FString& Addr : HostAddresses)
+	if (StatsObject == nullptr && UObjectInitialized())
 	{
-		if (Addr.Contains(TEXT("platform://")))
-		{
-			PlatformAddress = Addr;
-		}
+		StatsObject = NewObject<UDebugStorageServerConnection>();
+		StatsObject->SetHostAddress(CurrentHostAddr);
+		StatsObject->AddToRoot();
+		StatsObject->StartDrawing();
 	}
 
-	if (!PlatformAddress.IsEmpty())
-	{
-		ConnectionBackend = MakeUnique<FStorageServerPlatformConnectionBackend>(*this);
-	}
-	else
+	if (StatsObject)
 	{
-		ConnectionBackend = MakeUnique<FStorageSocketConnectionBackend>(*this, *ISocketSubsystem::Get());
+		StatsObject->AddTimingInstance(Duration, Bytes);
 	}
-
-	return ConnectionBackend != nullptr;
 }
 
-bool FStorageServerConnection::CreatePlatformBackend(const FString& HostAddresses, int32 Port)
+// TODO revive FStorageServerChunkBatchRequest
+#if 0
+
+class FStorageServerChunkBatchRequest : private FStorageServerRequest
 {
-	ConnectionBackend = MakeUnique<FStorageServerPlatformConnectionBackend>(*this);
-	if (!ConnectionBackend)
-	{
-		return false;
-	}
+public:
+	FStorageServerChunkBatchRequest& AddChunk(const FIoChunkId& ChunkId, int64 Offset, int64 Size);
+	bool Issue(TFunctionRef<void(uint32 ChunkCount, uint32* ChunkIndices, uint64* ChunkSizes, FStorageServerResponse& ChunkDataStream)> OnResponse);
 
-	return true;
+private:
+	friend FStorageServerConnection;
+
+	FStorageServerChunkBatchRequest(FStorageServerConnection& Owner, FAnsiStringView Resource, FAnsiStringView Hostname);
+
+	FStorageServerConnection& Owner;
+	int32 ChunkCountOffset = 0;
+};
+
+FStorageServerChunkBatchRequest::FStorageServerChunkBatchRequest(FStorageServerConnection& InOwner, FAnsiStringView Resource, FAnsiStringView Hostname)
+	: FStorageServerRequest("POST", Resource, Hostname)
+	, Owner(InOwner)
+{
+	uint32 Magic = 0xAAAA'77AC;
+	uint32 ChunkCountPlaceHolder = 0;
+	uint32 Reserved1 = 0;
+	uint32 Reserved2 = 0;
+	*this << Magic;
+	ChunkCountOffset = BodyBuffer.Num();
+	*this << ChunkCountPlaceHolder << Reserved1 << Reserved2;
 }
 
-FStorageConnectionBackend* FStorageServerConnection::GetConnectionBackend() const
+FStorageServerChunkBatchRequest& FStorageServerChunkBatchRequest::AddChunk(const FIoChunkId& ChunkId, int64 Offset, int64 Size)
 {
-	return ConnectionBackend.Get();
+	uint32* ChunkCount = reinterpret_cast<uint32*>(BodyBuffer.GetData() + ChunkCountOffset);
+	*this << const_cast<FIoChunkId&>(ChunkId) << *ChunkCount << Offset << Size;
+	++(*ChunkCount);
+	return *this;
 }
 
-void FStorageServerConnection::AddTimingInstance(double duration, uint64 bytes)
+bool FStorageServerChunkBatchRequest::Issue(TFunctionRef<void(uint32 ChunkCount, uint32* ChunkIndices, uint64* ChunkSizes, FStorageServerResponse& ChunkDataStream)> OnResponse)
 {
-	if ((StatsObject == nullptr) && (UObjectInitialized()))
+	IStorageConnectionSocket* Socket = Send(Owner);
+	if (!Socket)
 	{
-		StatsObject = NewObject<UDebugStorageServerConnection>();
-		StatsObject->SetOwner(this);
-		StatsObject->AddToRoot();
-		StatsObject->StartDrawing();
+		UE_LOG(LogStorageServerConnection, Fatal, TEXT("Failed to send chunk batch request to storage server."));
+		return false;
+	}
+	FStorageServerResponse Response(Owner, *Socket);
+	if (!Response.IsOk())
+	{
+		UE_LOG(LogStorageServerConnection, Fatal, TEXT("Failed to read chunk batch from storage server. '%s'"), *Response.GetErrorMessage());
+		return false;
 	}
 
-	if (StatsObject)
-		StatsObject->AddTimingInstance(duration, bytes);
+	uint32 Magic;
+	uint32 ChunkCount;
+	uint32 Reserved1;
+	uint32 Reserved2;
+	Response << Magic;
+	if (Magic != 0xbada'b00f)
+	{
+		UE_LOG(LogStorageServerConnection, Fatal, TEXT("Invalid magic in chunk batch response from storage server."));
+		return false;
+	}
+	Response << ChunkCount;
+	if (ChunkCount > INT32_MAX)
+	{
+		UE_LOG(LogStorageServerConnection, Fatal, TEXT("Invalid chunk count in chunk batch response from storage server."));
+		return false;
+	}
+	Response << Reserved1;
+	Response << Reserved2;
+
+	TArray<uint32, TInlineAllocator<64>> ChunkIndices;
+	ChunkIndices.Reserve(ChunkCount);
+	TArray<uint64, TInlineAllocator<64>> ChunkSizes;
+	ChunkSizes.Reserve(ChunkCount);
+	for (uint32 Index = 0; Index < ChunkCount; ++Index)
+	{
+		uint32 ChunkIndex;
+		uint32 Flags;
+		int64 ChunkSize;
+		Response << ChunkIndex;
+		Response << Flags;
+		Response << ChunkSize;
+		ChunkIndices.Add(ChunkIndex);
+		ChunkSizes.Emplace(ChunkSize);
+	}
+	OnResponse(ChunkCount, ChunkIndices.GetData(), ChunkSizes.GetData(), Response);
+	Owner.AddTimingInstance(GetDuration(), (double)Response.Tell());
+	return true;
 }
 
+#endif
 
 #endif
diff --git a/Engine/Source/Runtime/StorageServerClient/Private/StorageServerConnection.h b/Engine/Source/Runtime/StorageServerClient/Private/StorageServerConnection.h
new file mode 100644
--- /dev/null
+++ b/Engine/Source/Runtime/StorageServerClient/Private/StorageServerConnection.h
@@ -0,0 +1,79 @@
+// Copyright Epic Games, Inc. All Rights Reserved.
+
+#pragma once
+
+#include "CoreTypes.h"
+#include "Compression/CompressedBuffer.h"
+#include "Containers/UnrealString.h"
+#include "Containers/AnsiString.h"
+#include "Misc/StringBuilder.h"
+#include "Memory/MemoryFwd.h"
+#include "Templates/SharedPointer.h"
+#include "SocketTypes.h"
+#include "HAL/PlatformTime.h"
+#include "StorageServerHttpClient.h"
+#include "DebugStorageServerConnection.h"
+#include "IO/IoChunkId.h"
+
+#if !UE_BUILD_SHIPPING
+
+DECLARE_LOG_CATEGORY_EXTERN(LogStorageServerConnection, Log, All);
+
+struct FPackageStoreEntryResource;
+
+class FStorageServerConnection
+{
+public:
+	FStorageServerConnection() = default;
+	~FStorageServerConnection();
+
+	bool Initialize(TArrayView<const FString> HostAddresses, const int32 Port, const FAnsiStringView& InBaseURI);
+
+	void PackageStoreRequest(TFunctionRef<void(FPackageStoreEntryResource&&)> Callback);
+	void FileManifestRequest(TFunctionRef<void(FIoChunkId Id, FStringView Path, int64 RawSize)> Callback);
+	int64 ChunkSizeRequest(const FIoChunkId& ChunkId);
+	TIoStatusOr<FIoBuffer> ReadChunkRequest(
+		const FIoChunkId& ChunkId,
+		const uint64 Offset,
+		const uint64 Size,
+		const TOptional<FIoBuffer> OptDestination,
+		const bool bHardwareTargetBuffer
+	);
+	void ReadChunkRequestAsync(
+		const FIoChunkId& ChunkId,
+		const uint64 Offset,
+		const uint64 Size,
+		const TOptional<FIoBuffer> OptDestination,
+		const bool bHardwareTargetBuffer,
+		TFunctionRef<void(TIoStatusOr<FIoBuffer> Data)> OnResponse
+	);
+
+	FString GetHostAddr() const
+	{
+		return CurrentHostAddr;
+	}
+
+private:
+	TUniquePtr<IStorageServerHttpClient> HttpClient;
+	FAnsiString BaseURI;
+	FString CurrentHostAddr;
+	UDebugStorageServerConnection* StatsObject = nullptr;
+
+	TArray<FString> SortHostAddressesByLocalSubnet(TArrayView<const FString> HostAddresses, const int32 Port);
+	static bool IsPlatformSocketAddress(const FString Address); 
+	TUniquePtr<IStorageServerHttpClient> CreateHttpClient(const FString Address, const int32 Port);
+	TSharedPtr<FInternetAddr> StringToInternetAddr(const FString Address, const int32 Port);
+	bool HandshakeRequest();
+	void BuildReadChunkRequestUrl(FAnsiStringBuilderBase& Builder, const FIoChunkId& ChunkId, const uint64 Offset, const uint64 Size);
+	static TIoStatusOr<FIoBuffer> ReadChunkRequestProcessHttpResult(
+		IStorageServerHttpClient::FResult ResultTuple,
+		const uint64 Offset,
+		const uint64 Size,
+		const TOptional<FIoBuffer> OptDestination,
+		const bool bHardwareTargetBuffer
+	);
+	static uint64 GetCompressedOffset(const FCompressedBuffer& Buffer, uint64 RawOffset);
+	void AddTimingInstance(const double Duration, const uint64 Bytes);
+};
+
+#endif
diff --git a/Engine/Source/Runtime/StorageServerClient/Private/StorageServerHttpClient.h b/Engine/Source/Runtime/StorageServerClient/Private/StorageServerHttpClient.h
new file mode 100644
--- /dev/null
+++ b/Engine/Source/Runtime/StorageServerClient/Private/StorageServerHttpClient.h
@@ -0,0 +1,50 @@
+// Copyright Epic Games, Inc. All Rights Reserved.
+
+#pragma once
+
+#include "CoreTypes.h"
+#include "IO/IoBuffer.h"
+#include "IO/IoStatus.h"
+#include "Misc/Optional.h"
+
+#if !UE_BUILD_SHIPPING
+
+enum class EStorageServerContentType : uint8
+{
+	Unknown = 0,
+	CbObject,
+	Binary,
+	CompressedBinary,
+};
+
+class IStorageServerHttpClient 
+{
+public:
+	using FResult = TTuple<TIoStatusOr<FIoBuffer>, EStorageServerContentType>;
+	using FResultCallback = TFunction<void(FResult)>;
+
+	virtual ~IStorageServerHttpClient() = default;
+
+	virtual FResult RequestSync(
+		FAnsiStringView Url,
+		EStorageServerContentType Accept = EStorageServerContentType::Unknown,
+		FAnsiStringView Verb = "GET",
+		TOptional<FIoBuffer> OptPayload = TOptional<FIoBuffer>(),
+		TOptional<FIoBuffer> OptDestination = TOptional<FIoBuffer>(),
+		float TimeoutSeconds = -1.f,
+		const bool bReportErrors = true
+	) = 0;
+
+	virtual void RequestAsync(
+		FResultCallback&& Callback,
+		FAnsiStringView Url,
+		EStorageServerContentType Accept = EStorageServerContentType::Unknown,
+		FAnsiStringView Verb = "GET",
+		TOptional<FIoBuffer> OptPayload = TOptional<FIoBuffer>(),
+		TOptional<FIoBuffer> OptDestination = TOptional<FIoBuffer>(),
+		float TimeoutSeconds = -1.f,
+		const bool bReportErrors = true
+	) = 0;
+};
+
+#endif
diff --git a/Engine/Source/Runtime/StorageServerClient/Private/StorageServerIoDispatcherBackend.cpp b/Engine/Source/Runtime/StorageServerClient/Private/StorageServerIoDispatcherBackend.cpp
--- a/Engine/Source/Runtime/StorageServerClient/Private/StorageServerIoDispatcherBackend.cpp
+++ b/Engine/Source/Runtime/StorageServerClient/Private/StorageServerIoDispatcherBackend.cpp
@@ -64,7 +64,7 @@ uint32 FStorageServerIoDispatcherBackend::Run()
 	const int32 BatchCount = GStorageServerIoDispatcherMaxActiveBatchCount;
 	for (int32 BatchIndex = 0; BatchIndex < BatchCount; ++BatchIndex)
 	{
-		FBatch* Batch = new FBatch(*this, MakeUnique<FStorageServerSerializationContext>());
+		FBatch* Batch = new FBatch(*this);
 		Batch->Next = FirstAvailableBatch;
 		FirstAvailableBatch = Batch;
 	}
@@ -148,20 +148,14 @@ bool FStorageServerIoDispatcherBackend::Resolve(FIoRequestImpl* Request)
 	else
 	{
 		TRACE_CPUPROFILER_EVENT_SCOPE(StorageServerIoDispatcherReadChunk);
-		bool bSuccess = Connection.ReadChunkRequest(Request->ChunkId, Request->Options.GetOffset(), Request->Options.GetSize(), [&Request](FStorageServerResponse& Response)
+		const TOptional<FIoBuffer> OptDestination = Request->Options.GetTargetVa() ? FIoBuffer(FIoBuffer::Wrap, Request->Options.GetTargetVa(), Request->Options.GetSize()) : TOptional<FIoBuffer>();
+		const bool bHardwareTargetBuffer = EnumHasAnyFlags(Request->Options.GetFlags(), EIoReadOptionsFlags::HardwareTargetBuffer);
+		TIoStatusOr<FIoBuffer> Result = Connection.ReadChunkRequest(Request->ChunkId, Request->Options.GetOffset(), Request->Options.GetSize(), OptDestination, bHardwareTargetBuffer); 
+		if (Result.IsOk())
 		{
-			TRACE_CPUPROFILER_EVENT_SCOPE(SerializeResponse);
-			FIoBuffer Chunk;
-			if (Response.SerializeChunk(Chunk, Request->Options.GetTargetVa(), Request->Options.GetOffset(), Request->Options.GetSize()))
-			{
-				Request->SetResult(Chunk);
-			}
-			else
-			{
-				Request->SetFailed();
-			}
-		});
-		if (!bSuccess)
+			Request->SetResult(Result.ValueOrDie());
+		}
+		else
 		{
 			Request->SetFailed();
 		}
@@ -278,9 +272,8 @@ void FStorageServerIoDispatcherBackend::OnBatchCompleted(FBatch* Batch)
 	BatchCompletedEvent->Trigger();
 }
 
-FStorageServerIoDispatcherBackend::FBatch::FBatch(FStorageServerIoDispatcherBackend& InOwner, TUniquePtr<FStorageServerSerializationContext> InSerializationContext)
+FStorageServerIoDispatcherBackend::FBatch::FBatch(FStorageServerIoDispatcherBackend& InOwner)
 	: Owner(InOwner)
-	, SerializationContext(MoveTemp(InSerializationContext))
 {
 
 }
@@ -337,22 +330,16 @@ void FStorageServerIoDispatcherBackend::FBatch::DoThreadedWork()
 	while (Request)
 	{
 		FIoRequestImpl* NextRequest = Request->NextRequest;
+
 		TRACE_CPUPROFILER_EVENT_SCOPE(StorageServerIoDispatcherReadChunk);
-		bool bSuccess = Owner.Connection.ReadChunkRequest(Request->ChunkId, Request->Options.GetOffset(), Request->Options.GetSize(), [this, &Request](FStorageServerResponse& Response)
-			{
-				TRACE_CPUPROFILER_EVENT_SCOPE(SerializeResponse);
-				FIoBuffer Chunk;
-				if (Response.SerializeChunk(*SerializationContext, Chunk, Request->Options.GetTargetVa(), Request->Options.GetOffset(), Request->Options.GetSize(), 
-					EnumHasAnyFlags(Request->Options.GetFlags(), EIoReadOptionsFlags::HardwareTargetBuffer) ? true : false ))
-				{
-					Request->SetResult(Chunk);
-				}
-				else
-				{
-					Request->SetFailed();
-				}
-			});
-		if (!bSuccess)
+		const TOptional<FIoBuffer> OptDestination = Request->Options.GetTargetVa() ? FIoBuffer(FIoBuffer::Wrap, Request->Options.GetTargetVa(), Request->Options.GetSize()) : TOptional<FIoBuffer>();
+		const bool bHardwareTargetBuffer = EnumHasAnyFlags(Request->Options.GetFlags(), EIoReadOptionsFlags::HardwareTargetBuffer);
+		TIoStatusOr<FIoBuffer> Result = Owner.Connection.ReadChunkRequest(Request->ChunkId, Request->Options.GetOffset(), Request->Options.GetSize(), OptDestination, bHardwareTargetBuffer); 
+		if (Result.IsOk())
+		{
+			Request->SetResult(Result.ValueOrDie());
+		}
+		else
 		{
 			Request->SetFailed();
 		}
diff --git a/Engine/Source/Runtime/StorageServerClient/Private/StorageServerIoDispatcherBackend.h b/Engine/Source/Runtime/StorageServerClient/Private/StorageServerIoDispatcherBackend.h
--- a/Engine/Source/Runtime/StorageServerClient/Private/StorageServerIoDispatcherBackend.h
+++ b/Engine/Source/Runtime/StorageServerClient/Private/StorageServerIoDispatcherBackend.h
@@ -66,7 +66,7 @@ private:
 	struct FBatch
 		: public IQueuedWork
 	{
-		FBatch(FStorageServerIoDispatcherBackend& InOwner, TUniquePtr<FStorageServerSerializationContext> InSerializationContext);
+		FBatch(FStorageServerIoDispatcherBackend& InOwner);
 
 		virtual void DoThreadedWork() override;
 		virtual void Abandon() override {};
@@ -76,7 +76,6 @@ private:
 		FIoRequestImpl* RequestsHead = nullptr;
 		FIoRequestImpl* RequestsTail = nullptr;
 		uint64 RequestsCount = 0;
-		TUniquePtr<FStorageServerSerializationContext> SerializationContext;
 	};
 
 	bool Resolve(FIoRequestImpl* Request);
diff --git a/Engine/Source/Runtime/StorageServerClient/Private/StorageServerPlatformFile.cpp b/Engine/Source/Runtime/StorageServerClient/Private/StorageServerPlatformFile.cpp
--- a/Engine/Source/Runtime/StorageServerClient/Private/StorageServerPlatformFile.cpp
+++ b/Engine/Source/Runtime/StorageServerClient/Private/StorageServerPlatformFile.cpp
@@ -1003,12 +1003,8 @@ FFileStatData FStorageServerPlatformFile::SendGetStatDataMessage(const FIoChunkI
 int64 FStorageServerPlatformFile::SendReadMessage(uint8* Destination, const FIoChunkId& FileChunkId, int64 Offset, int64 BytesToRead)
 {
 	TRACE_CPUPROFILER_EVENT_SCOPE(StorageServerPlatformFileRead);
-	int64 BytesRead = 0;
-	Connection->ReadChunkRequest(FileChunkId, Offset, BytesToRead, [Destination, Offset, BytesToRead, &BytesRead](FStorageServerResponse& Response)
-	{
-		BytesRead = Response.SerializeChunkTo(MakeMemoryView(Destination, BytesToRead), Offset);
-	});
-	return BytesRead;
+	TIoStatusOr<FIoBuffer> Result = Connection->ReadChunkRequest(FileChunkId, Offset, BytesToRead, FIoBuffer(FIoBuffer::Wrap, Destination, BytesToRead), false);
+	return Result.IsOk() ? Result.ValueOrDie().GetSize() : 0;
 }
 
 bool FStorageServerPlatformFile::SendMessageToServer(const TCHAR* Message, IPlatformFile::IFileServerMessageHandler* Handler)
diff --git a/Engine/Source/Runtime/StorageServerClient/Private/StorageSocketConnectionBackend.cpp b/Engine/Source/Runtime/StorageServerClient/Private/StorageSocketConnectionBackend.cpp
deleted file mode 100644
diff --git a/Engine/Source/Runtime/StorageServerClient/Private/StorageSocketConnectionBackend.h b/Engine/Source/Runtime/StorageServerClient/Private/StorageSocketConnectionBackend.h
deleted file mode 100644
diff --git a/Engine/Source/Runtime/StorageServerClient/Public/StorageServerConnection.h b/Engine/Source/Runtime/StorageServerClient/Public/StorageServerConnection.h
deleted file mode 100644
