Zen: Fix for crashes when looking for the server address on platforms that don't support IP6.

diff --git a/Engine/Source/Runtime/StorageServerClient/Private/StorageServerConnection.cpp b/Engine/Source/Runtime/StorageServerClient/Private/StorageServerConnection.cpp
--- a/Engine/Source/Runtime/StorageServerClient/Private/StorageServerConnection.cpp
+++ b/Engine/Source/Runtime/StorageServerClient/Private/StorageServerConnection.cpp
@@ -259,6 +259,12 @@ TArray<FString> FStorageServerConnection::SortHostAddressesByLocalSubnet(TArrayV
 		}
 
 		TSharedPtr<FInternetAddr> Addr = StringToInternetAddr(HostAddress, Port);
+
+		if (!Addr)
+		{
+			continue;
+		}
+
 		FString tempAddrStringSubnet = Addr->ToString(bAppendPort);
 
 #if PLATFORM_DESKTOP || PLATFORM_ANDROID
