[HttpClient]
 - Fix crash when dealing with empty content

diff --git a/Engine/Source/Runtime/StorageServerClient/Private/BuiltInHttpClient/BuiltInHttpClient.cpp b/Engine/Source/Runtime/StorageServerClient/Private/BuiltInHttpClient/BuiltInHttpClient.cpp
--- a/Engine/Source/Runtime/StorageServerClient/Private/BuiltInHttpClient/BuiltInHttpClient.cpp
+++ b/Engine/Source/Runtime/StorageServerClient/Private/BuiltInHttpClient/BuiltInHttpClient.cpp
@@ -162,7 +162,7 @@ IStorageServerHttpClient::FResult FBuiltInHttpClient::RequestSync(
 
 	TOptional<FIoBuffer> ResponsePayload = ContentLength > 0
 		? (OptDestination.IsSet() && ((int64)OptDestination->GetSize() >= ContentLength) ? OptDestination.GetValue() : FIoBuffer(ContentLength))
-		: TOptional<FIoBuffer>();
+		: FIoBuffer(0);
 
 	if (ContentLength)
 	{
