Added the option to specify file extensions that will be assumed to be immutable, and their timestamps will always be constant.

 - only .uplugin at the moment
 - fixes the GFP plugin data caching when using Zen

#rnx

diff --git a/Engine/Source/Runtime/StorageServerClient/Private/StorageServerPlatformFile.cpp b/Engine/Source/Runtime/StorageServerClient/Private/StorageServerPlatformFile.cpp
--- a/Engine/Source/Runtime/StorageServerClient/Private/StorageServerPlatformFile.cpp
+++ b/Engine/Source/Runtime/StorageServerClient/Private/StorageServerPlatformFile.cpp
@@ -31,6 +31,8 @@ DEFINE_LOG_CATEGORY_STATIC(LogStorageServerPlatformFile, Log, All);
 #define EXCLUDE_NONSERVER_UE_EXTENSIONS 1	// Use .Build.cs file to disable this if the game relies on accessing loose files on the local filesystem
 #endif
 
+static FDateTime GAssumedImmutableTimeStamp = FDateTime::Now();
+
 FStorageServerFileSystemTOC::~FStorageServerFileSystemTOC()
 {
 	FWriteScopeLock _(TocLock);
@@ -445,6 +447,12 @@ bool FStorageServerPlatformFile::Initialize(IPlatformFile* Inner, const TCHAR* C
 		ExcludedNonServerExtensions.Add(TEXT("ushaderbytecode"));
 		ExcludedNonServerExtensions.Add(TEXT("ini")); //special cases of local only ini file needs to be managed as special exclusion
 #endif
+
+#if !WITH_EDITOR
+		// Extensions for file types that will be assumed to be immutable - their time stamp will remain unchanged.
+		AssumedImmutableTimeStampExtensions.Add(TEXT("uplugin"));
+#endif
+
 		// Don't initialize the connection yet because we want to incorporate project file path information into the initialization.
 
 		TUniquePtr<FArchive> ProjectStoreMarkerReader = TryFindProjectStoreMarkerFile(Inner);
@@ -558,7 +566,7 @@ FDateTime FStorageServerPlatformFile::GetTimeStamp(const TCHAR* Filename)
 	{
 		if (ServerToc.FileExists(*StorageServerFilename))
 		{
-			return FDateTime::Now();
+			return IsAssumedImmutableTimeStampFilename(*StorageServerFilename) ? GAssumedImmutableTimeStamp : FDateTime::Now();
 		}
 	}
 	return IsNonServerFilenameAllowed(Filename) ? LowerLevel->GetTimeStamp(Filename) : FDateTime::MinValue();
@@ -571,7 +579,7 @@ FDateTime FStorageServerPlatformFile::GetAccessTimeStamp(const TCHAR* Filename)
 	{
 		if (ServerToc.FileExists(*StorageServerFilename))
 		{
-			return FDateTime::Now();
+			return IsAssumedImmutableTimeStampFilename(*StorageServerFilename) ? GAssumedImmutableTimeStamp : FDateTime::Now();
 		}
 	}
 	return IsNonServerFilenameAllowed(Filename) ? LowerLevel->GetAccessTimeStamp(Filename) : FDateTime::MinValue();
@@ -925,6 +933,12 @@ bool FStorageServerPlatformFile::IsNonServerFilenameAllowed(FStringView InFilena
 	return bAllowed;
 }
 
+bool FStorageServerPlatformFile::IsAssumedImmutableTimeStampFilename(FStringView InFilename) const
+{
+	FName Ext = FName(FPathViews::GetExtension(InFilename));
+	return AssumedImmutableTimeStampExtensions.Contains(Ext);
+}
+
 bool FStorageServerPlatformFile::MakeStorageServerPath(const TCHAR* LocalFilenameOrDirectory, FStringBuilderBase& OutPath) const
 {
 	FStringView LocalEngineDirView(FPlatformMisc::EngineDir());
diff --git a/Engine/Source/Runtime/StorageServerClient/Private/StorageServerPlatformFile.h b/Engine/Source/Runtime/StorageServerClient/Private/StorageServerPlatformFile.h
--- a/Engine/Source/Runtime/StorageServerClient/Private/StorageServerPlatformFile.h
+++ b/Engine/Source/Runtime/StorageServerClient/Private/StorageServerPlatformFile.h
@@ -109,6 +109,7 @@ private:
 	friend class FStorageServerFileHandle;
 
 	bool IsNonServerFilenameAllowed(FStringView InFilename);
+	bool IsAssumedImmutableTimeStampFilename(FStringView InFilename) const;
 	bool MakeStorageServerPath(const TCHAR* LocalFilenameOrDirectory, FStringBuilderBase& OutPath) const;
 	bool MakeLocalPath(const TCHAR* ServerFilenameOrDirectory, FStringBuilderBase& OutPath) const;
 	IFileHandle* InternalOpenFile(const FIoChunkId& FileChunkId, int64 RawSize, const TCHAR* LocalFilename);
@@ -122,6 +123,7 @@ private:
 	FAnsiString MakeBaseURI();
 
 	TSet<FName> ExcludedNonServerExtensions;
+	TSet<FName> AssumedImmutableTimeStampExtensions;
 	IPlatformFile* LowerLevel = nullptr;
 	FStringView ServerEngineDirView = FStringView(TEXT("/{engine}/"));
 	FStringView ServerProjectDirView = FStringView(TEXT("/{project}/"));
