Ensure that a null lower level platform file doesn't cause StorageServerPlatformFile to crash with a null dereference.

diff --git a/Engine/Source/Runtime/StorageServerClient/Private/StorageServerPlatformFile.cpp b/Engine/Source/Runtime/StorageServerClient/Private/StorageServerPlatformFile.cpp
--- a/Engine/Source/Runtime/StorageServerClient/Private/StorageServerPlatformFile.cpp
+++ b/Engine/Source/Runtime/StorageServerClient/Private/StorageServerPlatformFile.cpp
@@ -587,7 +587,7 @@ bool FStorageServerPlatformFile::FileExists(const TCHAR* Filename)
 		return true;
 	}
 
-	return IsNonServerFilenameAllowed(Filename) ? LowerLevel->FileExists(Filename) : false;
+	return (LowerLevel && IsNonServerFilenameAllowed(Filename)) ? LowerLevel->FileExists(Filename) : false;
 }
 
 FDateTime FStorageServerPlatformFile::GetTimeStamp(const TCHAR* Filename)
@@ -600,7 +600,7 @@ FDateTime FStorageServerPlatformFile::GetTimeStamp(const TCHAR* Filename)
 			return IsAssumedImmutableTimeStampFilename(*StorageServerFilename) ? GAssumedImmutableTimeStamp : FDateTime::Now();
 		}
 	}
-	return IsNonServerFilenameAllowed(Filename) ? LowerLevel->GetTimeStamp(Filename) : FDateTime::MinValue();
+	return (LowerLevel && IsNonServerFilenameAllowed(Filename)) ? LowerLevel->GetTimeStamp(Filename) : FDateTime::MinValue();
 }
 
 FDateTime FStorageServerPlatformFile::GetAccessTimeStamp(const TCHAR* Filename)
@@ -613,7 +613,7 @@ FDateTime FStorageServerPlatformFile::GetAccessTimeStamp(const TCHAR* Filename)
 			return IsAssumedImmutableTimeStampFilename(*StorageServerFilename) ? GAssumedImmutableTimeStamp : FDateTime::Now();
 		}
 	}
-	return IsNonServerFilenameAllowed(Filename) ? LowerLevel->GetAccessTimeStamp(Filename) : FDateTime::MinValue();
+	return (LowerLevel && IsNonServerFilenameAllowed(Filename)) ? LowerLevel->GetAccessTimeStamp(Filename) : FDateTime::MinValue();
 }
 
 int64 FStorageServerPlatformFile::FileSize(const TCHAR* Filename)
@@ -627,7 +627,7 @@ int64 FStorageServerPlatformFile::FileSize(const TCHAR* Filename)
 			return FileSize;
 		}
 	}
-	return IsNonServerFilenameAllowed(Filename) ? LowerLevel->FileSize(Filename) : STORAGE_SERVER_FILE_UNKOWN_SIZE;
+	return (LowerLevel && IsNonServerFilenameAllowed(Filename)) ? LowerLevel->FileSize(Filename) : STORAGE_SERVER_FILE_UNKOWN_SIZE;
 }
 
 bool FStorageServerPlatformFile::IsReadOnly(const TCHAR* Filename)
@@ -637,7 +637,7 @@ bool FStorageServerPlatformFile::IsReadOnly(const TCHAR* Filename)
 	{
 		return true;
 	}
-	return IsNonServerFilenameAllowed(Filename) ? LowerLevel->IsReadOnly(Filename) : false;
+	return (LowerLevel && IsNonServerFilenameAllowed(Filename)) ? LowerLevel->IsReadOnly(Filename) : false;
 }
 
 FFileStatData FStorageServerPlatformFile::GetStatData(const TCHAR* FilenameOrDirectory)
@@ -668,7 +668,7 @@ FFileStatData FStorageServerPlatformFile::GetStatData(const TCHAR* FilenameOrDir
 		}
 	}
 	FFileStatData FileStatData;
-	if (IsNonServerFilenameAllowed(FilenameOrDirectory))
+	if (LowerLevel && IsNonServerFilenameAllowed(FilenameOrDirectory))
 	{
 		FileStatData = LowerLevel->GetStatData(FilenameOrDirectory);
 	}
@@ -692,7 +692,7 @@ IFileHandle* FStorageServerPlatformFile::OpenRead(const TCHAR* Filename, bool bA
 			return InternalOpenFile(FileChunkId, RawSize, Filename);
 		}
 	}
-	return IsNonServerFilenameAllowed(Filename) ? LowerLevel->OpenRead(Filename, bAllowWrite) : nullptr;
+	return (LowerLevel && IsNonServerFilenameAllowed(Filename)) ? LowerLevel->OpenRead(Filename, bAllowWrite) : nullptr;
 }
 
 bool FStorageServerPlatformFile::IterateDirectory(const TCHAR* Directory, IPlatformFile::FDirectoryVisitor& Visitor)
@@ -710,7 +710,7 @@ bool FStorageServerPlatformFile::IterateDirectory(const TCHAR* Directory, IPlatf
 			return Visitor.CallShouldVisitAndVisit(*LocalPath, bDirectory);
 		});
 	}
-	else
+	else if (LowerLevel)
 	{
 		bResult |= LowerLevel->IterateDirectory(Directory, Visitor);
 	}
@@ -753,7 +753,7 @@ bool FStorageServerPlatformFile::IterateDirectoryStat(const TCHAR* Directory, FD
 			return Visitor.CallShouldVisitAndVisit(*LocalPath, FileStatData);
 		});
 	}
-	else
+	else if (LowerLevel)
 	{
 		bResult |= LowerLevel->IterateDirectoryStat(Directory, Visitor);
 	}
@@ -762,7 +762,7 @@ bool FStorageServerPlatformFile::IterateDirectoryStat(const TCHAR* Directory, FD
 
 IMappedFileHandle* FStorageServerPlatformFile::OpenMapped(const TCHAR* Filename)
 {
-	return IsNonServerFilenameAllowed(Filename) ? LowerLevel->OpenMapped(Filename) : nullptr;
+	return (LowerLevel && IsNonServerFilenameAllowed(Filename)) ? LowerLevel->OpenMapped(Filename) : nullptr;
 }
 
 bool FStorageServerPlatformFile::DirectoryExists(const TCHAR* Directory)
@@ -772,7 +772,7 @@ bool FStorageServerPlatformFile::DirectoryExists(const TCHAR* Directory)
 	{
 		return true;
 	}
-	return LowerLevel->DirectoryExists(Directory);
+	return LowerLevel && LowerLevel->DirectoryExists(Directory);
 }
 
 FString FStorageServerPlatformFile::GetFilenameOnDisk(const TCHAR* Filename)
@@ -783,7 +783,7 @@ FString FStorageServerPlatformFile::GetFilenameOnDisk(const TCHAR* Filename)
 		UE_LOG(LogStorageServerPlatformFile, Warning, TEXT("Attempting to get disk filename of remote file '%s'"), Filename);
 		return Filename;
 	}
-	return IsNonServerFilenameAllowed(Filename) ? LowerLevel->GetFilenameOnDisk(Filename) : Filename;
+	return (LowerLevel && IsNonServerFilenameAllowed(Filename)) ? LowerLevel->GetFilenameOnDisk(Filename) : Filename;
 }
 
 bool FStorageServerPlatformFile::DeleteFile(const TCHAR* Filename)
@@ -793,11 +793,16 @@ bool FStorageServerPlatformFile::DeleteFile(const TCHAR* Filename)
 	{
 		return false;
 	}
-	return LowerLevel->DeleteFile(Filename);
+	return LowerLevel && LowerLevel->DeleteFile(Filename);
 }
 
 bool FStorageServerPlatformFile::MoveFile(const TCHAR* To, const TCHAR* From)
 {
+	if (!LowerLevel)
+	{
+		return false;
+	}
+
 	TStringBuilder<1024> StorageServerTo;
 	if (MakeStorageServerPath(To, StorageServerTo) && ServerToc.FileExists(*StorageServerTo))
 	{
@@ -851,7 +856,7 @@ bool FStorageServerPlatformFile::SetReadOnly(const TCHAR* Filename, bool bNewRea
 	{
 		return bNewReadOnlyValue;
 	}
-	return LowerLevel->SetReadOnly(Filename, bNewReadOnlyValue);
+	return LowerLevel && LowerLevel->SetReadOnly(Filename, bNewReadOnlyValue);
 }
 
 void FStorageServerPlatformFile::SetTimeStamp(const TCHAR* Filename, FDateTime DateTime)
@@ -861,7 +866,10 @@ void FStorageServerPlatformFile::SetTimeStamp(const TCHAR* Filename, FDateTime D
 	{
 		return;
 	}
-	LowerLevel->SetTimeStamp(Filename, DateTime);
+	if (LowerLevel)
+	{
+		LowerLevel->SetTimeStamp(Filename, DateTime);
+	}
 }
 
 IFileHandle* FStorageServerPlatformFile::OpenWrite(const TCHAR* Filename, bool bAppend, bool bAllowRead)
@@ -871,7 +879,11 @@ IFileHandle* FStorageServerPlatformFile::OpenWrite(const TCHAR* Filename, bool b
 	{
 		return nullptr;
 	}
-	return LowerLevel->OpenWrite(Filename, bAppend, bAllowRead);
+	if (LowerLevel)
+	{
+		LowerLevel->OpenWrite(Filename, bAppend, bAllowRead);
+	}
+	return nullptr;
 }
 
 bool FStorageServerPlatformFile::CreateDirectory(const TCHAR* Directory)
@@ -881,7 +893,7 @@ bool FStorageServerPlatformFile::CreateDirectory(const TCHAR* Directory)
 	{
 		return true;
 	}
-	return LowerLevel->CreateDirectory(Directory);
+	return LowerLevel && LowerLevel->CreateDirectory(Directory);
 }
 
 bool FStorageServerPlatformFile::DeleteDirectory(const TCHAR* Directory)
@@ -891,7 +903,7 @@ bool FStorageServerPlatformFile::DeleteDirectory(const TCHAR* Directory)
 	{
 		return false;
 	}
-	return LowerLevel->DeleteDirectory(Directory);
+	return LowerLevel && LowerLevel->DeleteDirectory(Directory);
 }
 
 FString FStorageServerPlatformFile::ConvertToAbsolutePathForExternalAppForRead(const TCHAR* Filename)
@@ -930,14 +942,19 @@ FString FStorageServerPlatformFile::ConvertToAbsolutePathForExternalAppForRead(c
 	if (PTRINT(DotSlashSkip - Filename) == 9) // 9 == ../../../
 	{
 		Result << DotSlashSkip;
-		if (LowerLevel->FileExists(Result.ToString()))
+		if (LowerLevel && LowerLevel->FileExists(Result.ToString()))
 		{
 			return FString::ConstructFromPtrSize(Result.GetData(), Result.Len());
 		}
 	}
 #endif
 
-	return LowerLevel->ConvertToAbsolutePathForExternalAppForRead(Filename);
+	if (LowerLevel)
+	{
+		return LowerLevel->ConvertToAbsolutePathForExternalAppForRead(Filename);
+	}
+
+	return IStorageServerPlatformFile::ConvertToAbsolutePathForExternalAppForRead(Filename);
 }
 
 bool FStorageServerPlatformFile::IsNonServerFilenameAllowed(FStringView InFilename)
