[Trivial] fix for some zen file reads not being properly reported in insights I/O

#rnx

diff --git a/Engine/Source/Runtime/StorageServerClient/Private/StorageServerPlatformFile.cpp b/Engine/Source/Runtime/StorageServerClient/Private/StorageServerPlatformFile.cpp
--- a/Engine/Source/Runtime/StorageServerClient/Private/StorageServerPlatformFile.cpp
+++ b/Engine/Source/Runtime/StorageServerClient/Private/StorageServerPlatformFile.cpp
@@ -259,7 +259,7 @@ public:
 			if (BytesReadFromBuffer == BytesToRead)
 			{
 				FilePos += BytesReadFromBuffer;
-				TRACE_PLATFORMFILE_END_READ(this, BytesReadFromBuffer);
+				TRACE_PLATFORMFILE_END_READ(Destination, BytesReadFromBuffer);
 				return true;
 			}
 		}
