[Trivial] fix copy/paste storage server error message

#rnx

diff --git a/Engine/Source/Runtime/StorageServerClient/Private/StorageServerConnection.cpp b/Engine/Source/Runtime/StorageServerClient/Private/StorageServerConnection.cpp
--- a/Engine/Source/Runtime/StorageServerClient/Private/StorageServerConnection.cpp
+++ b/Engine/Source/Runtime/StorageServerClient/Private/StorageServerConnection.cpp
@@ -467,7 +467,7 @@ TIoStatusOr<FIoBuffer> FStorageServerConnection::ReadChunkRequestProcessHttpResu
 	EStorageServerContentType MimeType = ResultTuple.Get<1>();
 	if (!Result.IsOk())
 	{
-		UE_LOG(LogStorageServerConnection, Fatal, TEXT("Failed to get chunk size from storage server. '%s'"), *Result.Status().ToString());
+		UE_LOG(LogStorageServerConnection, Fatal, TEXT("Failed read chunk from storage server. '%s'  Offset:%ull  Size:%ull"), *Result.Status().ToString(), Offset, Size);
 		return Result.Status();
 	}
 
