[Trivial] fix for missing return value for LowerLevel OpenWrite calls

diff --git a/Engine/Source/Runtime/StorageServerClient/Private/StorageServerPlatformFile.cpp b/Engine/Source/Runtime/StorageServerClient/Private/StorageServerPlatformFile.cpp
--- a/Engine/Source/Runtime/StorageServerClient/Private/StorageServerPlatformFile.cpp
+++ b/Engine/Source/Runtime/StorageServerClient/Private/StorageServerPlatformFile.cpp
@@ -881,7 +881,7 @@ IFileHandle* FStorageServerPlatformFile::OpenWrite(const TCHAR* Filename, bool b
 	}
 	if (LowerLevel)
 	{
-		LowerLevel->OpenWrite(Filename, bAppend, bAllowRead);
+		return LowerLevel->OpenWrite(Filename, bAppend, bAllowRead);
 	}
 	return nullptr;
 }
