Ensure that a missing oplog doesn't prevent us from starting editor/cooker when using zen as cook store and reader.

diff --git a/Engine/Source/Runtime/StorageServerClient/Private/StorageServerClientModule.cpp b/Engine/Source/Runtime/StorageServerClient/Private/StorageServerClientModule.cpp
--- a/Engine/Source/Runtime/StorageServerClient/Private/StorageServerClientModule.cpp
+++ b/Engine/Source/Runtime/StorageServerClient/Private/StorageServerClientModule.cpp
@@ -53,6 +53,7 @@ IStorageServerPlatformFile* FStorageServerClientModule::TryCreateCustomPlatformF
 	TUniquePtr<FStorageServerPlatformFile> StorageServerPlatformFile = MakeUnique<FStorageServerPlatformFile>();
 	StorageServerPlatformFile->SetCustomProjectStorePath(StoreDirectory);
 	StorageServerPlatformFile->SetAllowPackageIo(false);
+	StorageServerPlatformFile->SetAbortOnConnectionFailure(false);
 	const TCHAR* CmdLine = FCommandLine::Get();
 	if (StorageServerPlatformFile->ShouldBeUsed(Inner, CmdLine) && StorageServerPlatformFile->Initialize(Inner, CmdLine))
 	{
diff --git a/Engine/Source/Runtime/StorageServerClient/Private/StorageServerPlatformFile.cpp b/Engine/Source/Runtime/StorageServerClient/Private/StorageServerPlatformFile.cpp
--- a/Engine/Source/Runtime/StorageServerClient/Private/StorageServerPlatformFile.cpp
+++ b/Engine/Source/Runtime/StorageServerClient/Private/StorageServerPlatformFile.cpp
@@ -558,7 +558,7 @@ void FStorageServerPlatformFile::InitializeAfterProjectFilePath()
 			FModuleManager::Get().LoadModule("StorageServerClientDebug");
 		}
 	}
-	else
+	else if (bAbortOnConnectionFailure)
 	{
 		if (!FApp::IsUnattended())
 		{
@@ -566,7 +566,7 @@ void FStorageServerPlatformFile::InitializeAfterProjectFilePath()
 			FString FailedConnectionText = FString::Printf(TEXT(
 				"Network data streaming failed to connect to any of the following data sources:\n\n%s\n\n"
 				"This can be due to the sources being offline, the Unreal Zen Storage process not currently running, "
-				"invalid addresses, firewall blocking, or the sources being on a different network from this device. "
+				"invalid addresses, firewall blocking, or the sources being on a different network from this device.\n"
 				"Please verify that your Unreal Zen Storage process is running using the ZenDashboard utility. "
 				"If these issues can't be addressed, you can use an installed build without network data streaming by "
 				"building with the '-pak' argument. This process will now exit."),
diff --git a/Engine/Source/Runtime/StorageServerClient/Private/StorageServerPlatformFile.h b/Engine/Source/Runtime/StorageServerClient/Private/StorageServerPlatformFile.h
--- a/Engine/Source/Runtime/StorageServerClient/Private/StorageServerPlatformFile.h
+++ b/Engine/Source/Runtime/StorageServerClient/Private/StorageServerPlatformFile.h
@@ -113,6 +113,10 @@ public:
 	{
 		bAllowPackageIo = bInAllowPackageIo;
 	}
+	void SetAbortOnConnectionFailure(bool bInAbortOnConnectionFailure)
+	{
+		bAbortOnConnectionFailure = bInAbortOnConnectionFailure;
+	}
 	void SetCustomProjectStorePath(FStringView InProjectStorePath)
 	{
 		CustomProjectStorePath = InProjectStorePath;
@@ -151,6 +155,7 @@ private:
 	mutable TArray<FString> HostAddrs;
 	mutable uint16 HostPort = 8558;
 	bool bAllowPackageIo = true;
+	bool bAbortOnConnectionFailure = true;
 };
 
 #endif
