Use shared cache for storage server platform file.

#rnx

diff --git a/Engine/Source/Runtime/Core/Public/Containers/LruCache.h b/Engine/Source/Runtime/Core/Public/Containers/LruCache.h
--- a/Engine/Source/Runtime/Core/Public/Containers/LruCache.h
+++ b/Engine/Source/Runtime/Core/Public/Containers/LruCache.h
@@ -66,6 +66,18 @@ class TLruCache
 			, Value(InValue)
 		{ }
 
+		/**
+		 * Create a new instance with a default key value.
+		 *
+		 * @param InKey The entry's key.
+		 */
+		FCacheEntry(const KeyType& InKey)
+			: Key(InKey)
+			, LessRecent(nullptr)
+			, MoreRecent(nullptr)
+		{ }
+
+
 		/** Add this entry before the given one. */
 		FORCEINLINE void LinkBefore(FCacheEntry* Other)
 		{
@@ -153,7 +165,7 @@ public:
 	 *
 	 * @param Key The entry's lookup key.
 	 * @param Value The entry's value.
-	 * @see Empty, Find, GetKeys, Remove
+	 * @see AddUninitialized_GetRef,Empty, Find, GetKeys, Remove
 	 */
 	void Add(const KeyType& Key, const ValueType& Value)
 	{
@@ -191,6 +203,56 @@ public:
 		}
 	}
 
+	/**
+	 * Add an entry to the cache.
+	 *
+	 * If an entry with the specified key already exists in the cache,
+	 * the value of the existing entry will be returned. The added or
+	 * updated entry will be marked as the most recently used one.
+	 *
+	 * @param Key The entry's lookup key.
+	 * @return The entry's value.
+	 * @see Add, Empty, Find, GetKeys, Remove
+	 */
+	ValueType& AddUninitialized_GetRef(const KeyType& Key)
+	{
+		check(MaxNumElements != 0 && "Cannot add values to zero size TLruCache");
+
+		FCacheEntry** EntryPtr = LookupSet.Find(Key);
+
+		if (EntryPtr != nullptr)
+		{
+			// update existing entry
+			FCacheEntry* Entry = *EntryPtr;
+			checkSlow(Entry->Key == Key);
+
+			MarkAsRecent(*Entry);
+			return Entry->Value;
+		}
+		else
+		{
+			// add new entry
+			if (LookupSet.Num() == MaxNumElements)
+			{
+				Remove(LeastRecent);
+			}
+
+			FCacheEntry* NewEntry = new FCacheEntry(Key);
+			NewEntry->LinkBefore(MostRecent);
+			MostRecent = NewEntry;
+
+			if (LeastRecent == nullptr)
+			{
+				LeastRecent = NewEntry;
+			}
+
+			LookupSet.Add(NewEntry);
+			return NewEntry->Value;
+		}
+	}
+
+
+
 	/**
 	 * Check whether an entry with the specified key is in the cache.
 	 *
diff --git a/Engine/Source/Runtime/StorageServerClient/Private/StorageServerPlatformFile.cpp b/Engine/Source/Runtime/StorageServerClient/Private/StorageServerPlatformFile.cpp
--- a/Engine/Source/Runtime/StorageServerClient/Private/StorageServerPlatformFile.cpp
+++ b/Engine/Source/Runtime/StorageServerClient/Private/StorageServerPlatformFile.cpp
@@ -19,9 +19,12 @@
 #include "Serialization/JsonReader.h"
 #include "Serialization/JsonSerializer.h"
 #include "ProfilingDebugging/PlatformFileTrace.h"
+#include "ProfilingDebugging/CountersTrace.h"
 #include "StorageServerConnection.h"
 #include "StorageServerIoDispatcherBackend.h"
 #include "StorageServerPackageStore.h"
+#include "Containers/LruCache.h"
+#include <atomic>
 
 DEFINE_LOG_CATEGORY_STATIC(LogStorageServerPlatformFile, Log, All);
 
@@ -160,21 +163,226 @@ bool FStorageServerFileSystemTOC::IterateDirectory(const FString& Path, TFunctio
 	return true;
 }
 
-class FStorageServerFileHandle
-	: public IFileHandle
+
+#if COUNTERSTRACE_ENABLED
+	TRACE_DECLARE_ATOMIC_FLOAT_COUNTER(StorageServerCache_HitRatioBytes, TEXT("ZenClient/FileCacheHitRatio"));
+	namespace
+	{
+		static std::atomic<uint64> CacheHitBytes = 0;
+		static std::atomic<uint64> CacheMissBytes = 0;
+	}
+
+	#define STORAGESERVER_CACHEMISS(Bytes) \
+	{\
+		CacheMissBytes += Bytes; \
+		TRACE_COUNTER_SET(StorageServerCache_HitRatioBytes, (double)CacheHitBytes / (double)(CacheMissBytes+CacheHitBytes) ); \
+	}
+
+	#define STORAGESERVER_CACHEHIT(Bytes) \
+	{\
+		CacheHitBytes += Bytes; \
+		TRACE_COUNTER_SET(StorageServerCache_HitRatioBytes, (double)CacheHitBytes / (double)(CacheMissBytes+CacheHitBytes) ); \
+	}
+
+#else
+
+	#define STORAGESERVER_CACHEMISS(Bytes)
+	#define STORAGESERVER_CACHEHIT(Bytes)
+
+#endif // COUNTERSTRACE_ENABLED
+
+
+
+class FStorageServerFileCache
 {
-	enum
+private:
+	typedef FIoChunkId CacheKey;
+
+	typedef DefaultKeyComparer<FIoChunkId> CacheKeyComparer;
+public:
+	// zen compression block size is often 256kb
+	static const int64 BlockSize = 256 * 1024;
+
+	// up to 4 mb cache, not counting temporary read buffers
+	static const uint32 MaxCacheElements = 16; 
+
+
+	struct CacheEntry
 	{
-		BufferSize = 64 << 10
+		int64 Start = -1;
+		TArray<uint8, TInlineAllocator<BlockSize>> Buffer;
+
+		FORCEINLINE int64 End()
+		{
+			return Start + Buffer.Num();
+		}
+
+		bool TryReadFromCache(int64& FilePos, uint8*& Destination, int64& BytesToRead, int64& BytesRead)
+		{
+			if (FilePos >= Start && FilePos < End())
+			{
+				BytesRead = FMath::Min(End() - FilePos, BytesToRead);
+				FMemory::Memcpy(Destination, Buffer.GetData() + FilePos - Start, BytesRead);
+				FilePos += BytesRead;
+				Destination += BytesRead;
+				BytesToRead -= BytesRead;
+				return true;
+			}
+			else
+			{
+				return false;
+			}
+		}
 	};
+
+	static FORCEINLINE int64 BlockOffset(int64 Position)
+	{
+		return (Position / BlockSize) * BlockSize;
+	}
+
+	static FStorageServerFileCache& Get()
+	{
+		static FStorageServerFileCache Instance;
+		return Instance;
+	}
+
+	void Lock()
+	{
+		CriticalSection.Lock();
+	}
+
+	void Unlock()
+	{
+		CriticalSection.Unlock();
+	}
+
+	CacheEntry& FindOrAdd(FIoChunkId FileChunkId)
+	{
+		CacheKey Key = FileChunkId;
+		if (const CacheEntry* ExistingEntry = Cache.FindAndTouch(Key))
+		{
+			return *const_cast<CacheEntry*>(ExistingEntry); // TODO change LRU cache API
+		}
+		else
+		{
+			CacheEntry& Entry = Cache.AddUninitialized_GetRef(Key);
+			Entry.Start = -1;
+			Entry.Buffer.Empty();
+
+			return Entry;
+		}
+	}
+
+	void ReadCached(FStorageServerConnection* Connection, FIoChunkId FileChunkId, int64& FilePos, uint8*& Destination, int64& BytesToRead)
+	{
+		if (BytesToRead == 0)
+		{
+			return;
+		}
+
+		// try to read existing data from cache
+		{
+			UE::TScopeLock Lock(*this);
+
+			CacheEntry& Entry = FindOrAdd(FileChunkId);
+			int64 BytesRead = 0;
+			if (Entry.TryReadFromCache(FilePos, Destination, BytesToRead, BytesRead))
+			{
+				STORAGESERVER_CACHEHIT(BytesRead);
+			}
+
+			if (BytesToRead == 0)
+			{
+				return;
+			}
+		}
+
+
+		// if request spans multiple blocks, satisfy all but last block without cache 
+		if (BlockOffset(FilePos) < BlockOffset(FilePos + BytesToRead))
+		{
+			const int64 BytesToReadRequested = BlockOffset(BytesToRead + FilePos) - FilePos;
+			const int64 BytesRead = SendReadMessage(Connection, Destination, FileChunkId, FilePos, BytesToReadRequested);
+			STORAGESERVER_CACHEMISS(BytesRead);
+			FilePos += BytesRead;
+			Destination += BytesRead;
+			BytesToRead -= BytesRead;
+		}
+
+		if (BytesToRead == 0)
+		{
+			return;
+		}
+
+		// try to read last block from cache
+		{
+			UE::TScopeLock Lock(*this);
+
+			CacheEntry& Entry = FindOrAdd(FileChunkId);
+			int64 BytesRead = 0;
+			if (Entry.TryReadFromCache(FilePos, Destination, BytesToRead, BytesRead))
+			{
+				STORAGESERVER_CACHEHIT(BytesRead);
+				if (ensure(BytesToRead == 0))
+				{
+					return;
+				}
+			}
+
+		}
+
+		// read and cache last block
+		// TODO try to avoid doing two requests for large reads 
+		{
+			TArray<uint8> TempBuffer; // allocating a temporary BlockSize buffer here for the read - one per parallel file access
+			TempBuffer.AddUninitialized(BlockSize);
+			int64 TempStart = BlockOffset(FilePos);
+
+			int64 BytesRead = SendReadMessage(Connection, TempBuffer.GetData(), FileChunkId, TempStart, TempBuffer.Num());
+			STORAGESERVER_CACHEMISS(BytesRead);
+
+			{
+				UE::TScopeLock Lock(*this);
+
+				CacheEntry& Entry = FindOrAdd(FileChunkId);
+				Entry.Start = TempStart;
+				Entry.Buffer.SetNum(BytesRead);
+				FMemory::Memcpy(Entry.Buffer.GetData(), TempBuffer.GetData(), BytesRead);
+
+				ensure(Entry.TryReadFromCache(FilePos, Destination, BytesToRead, BytesRead));
+			}
+		}
+
+		check(BytesToRead == 0);
+	}
+
+private:
+	FStorageServerFileCache()
+		: Cache(MaxCacheElements)
+	{
+	}
+
+	int64 SendReadMessage(FStorageServerConnection* Connection, uint8* Destination, const FIoChunkId& FileChunkId, int64 Offset, int64 BytesToRead)
+	{
+		TRACE_CPUPROFILER_EVENT_SCOPE(FStorageServerFileCache::SendReadMessage);
+		int64 BytesRead = 0;
+		TIoStatusOr<FIoBuffer> Result = Connection->ReadChunkRequest(FileChunkId, Offset, BytesToRead, FIoBuffer(FIoBuffer::Wrap, Destination, BytesToRead), false);
+		BytesRead = Result.IsOk() ? Result.ValueOrDie().GetSize() : 0;
+		return BytesRead;
+	}
+
+	TLruCache<CacheKey, CacheEntry, CacheKeyComparer> Cache;
+	FCriticalSection CriticalSection;
+};
+
+class FStorageServerFileHandle
+	: public IFileHandle
+{
 	FStorageServerPlatformFile& Owner;
 	FIoChunkId FileChunkId;
 	FString Filename;
 	int64 FilePos = 0;
 	int64 FileSize = -1;
-	int64 BufferStart = -1;
-	int64 BufferEnd = -1;
-	uint8 Buffer[BufferSize];
 
 public:
 	FStorageServerFileHandle(FStorageServerPlatformFile& InOwner, FIoChunkId InFileChunkId, int64 InFileSize, const TCHAR* InFilename)
@@ -236,50 +444,16 @@ public:
 			return true;
 		}
 
-		if (BytesToRead > BufferSize)
-		{
-			const int64 BytesRead = Owner.SendReadMessage(Destination, FileChunkId, FilePos, BytesToRead);
-			if (BytesRead == BytesToRead)
-			{
-				FilePos += BytesRead;
-				TRACE_PLATFORMFILE_END_READ(Destination, BytesRead);
-				return true;
-			}
-			TRACE_PLATFORMFILE_END_READ(Destination, 0);
-			return false;
-		}
+		FStorageServerFileCache& Cache = FStorageServerFileCache::Get();
 
-		int64 BytesReadFromBuffer = 0;
-		if (FilePos >= BufferStart && FilePos < BufferEnd)
-		{
-			const int64 BufferOffset = FilePos - BufferStart;
-			check(BufferOffset < BufferSize);
-			BytesReadFromBuffer = FMath::Min(BufferSize - BufferOffset, BytesToRead);
-			FMemory::Memcpy(Destination, Buffer + BufferOffset, BytesReadFromBuffer);
-			if (BytesReadFromBuffer == BytesToRead)
-			{
-				FilePos += BytesReadFromBuffer;
-				TRACE_PLATFORMFILE_END_READ(Destination, BytesReadFromBuffer);
-				return true;
-			}
-		}
+		uint8* DestinationPtr = Destination;
+		int64 BytesRemaining = BytesToRead;
+		Cache.ReadCached(Owner.Connection.Get(), FileChunkId, /*out*/FilePos, /*out*/DestinationPtr, /*out*/BytesRemaining);
+		int64 BytesRead = (BytesToRead - BytesRemaining);
 
-		const int64 BytesRead = Owner.SendReadMessage(Buffer, FileChunkId, FilePos + BytesReadFromBuffer, BufferSize);
-		BufferStart = FilePos + BytesReadFromBuffer;
-		BufferEnd = BufferStart + BytesRead;
-
-		const int64 BytesToReadFromBuffer = FMath::Min(BytesRead, BytesToRead - BytesReadFromBuffer);
-		FMemory::Memcpy(Destination + BytesReadFromBuffer, Buffer, BytesToReadFromBuffer);
-		BytesReadFromBuffer += BytesToReadFromBuffer;
-		if (BytesReadFromBuffer == BytesToRead)
-		{
-			FilePos += BytesReadFromBuffer;
-			TRACE_PLATFORMFILE_END_READ(Destination, BytesReadFromBuffer);
-			return true;
-		}
+		TRACE_PLATFORMFILE_END_READ(Destination, BytesRead);
 		
-		TRACE_PLATFORMFILE_END_READ(Destination, 0);
-		return false;
+		return BytesRemaining == 0;
 	}
 
 	virtual bool Write(const uint8* Source, int64 BytesToWrite) override
