[Trivial] tweak to storage server iostore cache hit rate trace name

diff --git a/Engine/Source/Runtime/StorageServerClient/Private/StorageServerPlatformFile.cpp b/Engine/Source/Runtime/StorageServerClient/Private/StorageServerPlatformFile.cpp
--- a/Engine/Source/Runtime/StorageServerClient/Private/StorageServerPlatformFile.cpp
+++ b/Engine/Source/Runtime/StorageServerClient/Private/StorageServerPlatformFile.cpp
@@ -165,7 +165,7 @@ bool FStorageServerFileSystemTOC::IterateDirectory(const FString& Path, TFunctio
 
 
 #if COUNTERSTRACE_ENABLED
-	TRACE_DECLARE_ATOMIC_FLOAT_COUNTER(StorageServerCache_HitRatioBytes, TEXT("ZenClient/FileCacheHitRatio"));
+	TRACE_DECLARE_ATOMIC_FLOAT_COUNTER(StorageServerCache_HitRatioBytes, TEXT("ZenClient/FileCache/HitRatio"));
 	namespace
 	{
 		static std::atomic<uint64> CacheHitBytes = 0;
