Zen pak file streaming:
 - -ZenPak used to indicate streaming using pak files instead of Zen loose file streaming
 - added pak file hardware decompression on supported platforms

diff --git a/Engine/Plugins/Experimental/DefaultInstallBundleManager/Source/Private/InstallBundleSourceBulk.cpp b/Engine/Plugins/Experimental/DefaultInstallBundleManager/Source/Private/InstallBundleSourceBulk.cpp
--- a/Engine/Plugins/Experimental/DefaultInstallBundleManager/Source/Private/InstallBundleSourceBulk.cpp
+++ b/Engine/Plugins/Experimental/DefaultInstallBundleManager/Source/Private/InstallBundleSourceBulk.cpp
@@ -3,6 +3,7 @@
 #include "InstallBundleSourceBulk.h"
 
 #include "DefaultInstallBundleManagerPrivate.h"
+#include "HAL/PlatformFileManager.h"
 #include "HAL/PlatformFile.h"
 #include "IPlatformFilePak.h"
 #include "InstallBundleManagerUtil.h"
@@ -270,13 +271,13 @@ void FInstallBundleSourceBulk::AsyncInit_MakeBundlesForBulkBuild()
 		
 		for (const FString& SearchDir : PakSearchDirs)
 		{
-			IPlatformFile::GetPlatformPhysical().FindFilesRecursively(*FoundFiles, *SearchDir, FileExtension);
+			FPlatformFileManager::Get().GetPlatformPhysical().FindFilesRecursively(*FoundFiles, *SearchDir, FileExtension);
 		}
 
 #if PLATFORM_IOS
 		// Only scan the root folder on IOS for shaderlibs.  Running this on windows is very expensive
 		// if the content dir contains loose assets which is common during development.
-		IPlatformFile::GetPlatformPhysical().FindFiles(*FoundFiles, *ContentDir, TEXT(".metallib"));
+		FPlatformFileManager::Get().GetPlatformPhysical().FindFiles(*FoundFiles, *ContentDir, TEXT(".metallib"));
 #endif // PLATFORM_IOS
 	},
 	[this, FoundFiles]()
diff --git a/Engine/Source/Runtime/Core/Private/GenericPlatform/GenericPlatformMisc.cpp b/Engine/Source/Runtime/Core/Private/GenericPlatform/GenericPlatformMisc.cpp
--- a/Engine/Source/Runtime/Core/Private/GenericPlatform/GenericPlatformMisc.cpp
+++ b/Engine/Source/Runtime/Core/Private/GenericPlatform/GenericPlatformMisc.cpp
@@ -1936,7 +1936,7 @@ FString FGenericPlatformMisc::LoadTextFileFromPlatformPackage(const FString& Rel
 {
 	FString Path = RootDir() / RelativePath;
 	FString Result;
-	if (FFileHelper::LoadFileToString(Result, &IPlatformFile::GetPlatformPhysical(), *Path))
+	if (FFileHelper::LoadFileToString(Result, &FPlatformFileManager::Get().GetPlatformPhysical(), *Path))
 	{
 		return Result;
 	}
diff --git a/Engine/Source/Runtime/Core/Private/Misc/CoreGlobals.cpp b/Engine/Source/Runtime/Core/Private/Misc/CoreGlobals.cpp
--- a/Engine/Source/Runtime/Core/Private/Misc/CoreGlobals.cpp
+++ b/Engine/Source/Runtime/Core/Private/Misc/CoreGlobals.cpp
@@ -983,6 +983,25 @@ int32 GetMultiprocessId()
 #endif
 }
 
+int GIsUsingZenPakFileStreaming = -1;
+
+bool IsUsingZenPakFileStreaming()
+{
+	if (GIsUsingZenPakFileStreaming == -1)
+	{
+		if (FParse::Param(FCommandLine::Get(), TEXT("ZenPak")))
+		{
+			GIsUsingZenPakFileStreaming = 1;
+		}
+		else
+		{
+			GIsUsingZenPakFileStreaming = 0;
+		}
+	}
+
+	return GIsUsingZenPakFileStreaming == 1;
+}
+
 }
 
 namespace UE::Private
diff --git a/Engine/Source/Runtime/Core/Public/CoreGlobals.h b/Engine/Source/Runtime/Core/Public/CoreGlobals.h
--- a/Engine/Source/Runtime/Core/Public/CoreGlobals.h
+++ b/Engine/Source/Runtime/Core/Public/CoreGlobals.h
@@ -832,3 +832,7 @@ extern CORE_API UE::Core::FVisualizerDebuggingState* GCoreDebuggingState;
 
 /** @return True if running cook-on-the-fly. */
 bool CORE_API IsRunningCookOnTheFly();
+
+namespace UE {
+	bool CORE_API IsUsingZenPakFileStreaming();
+}
\ No newline at end of file
diff --git a/Engine/Source/Runtime/Core/Public/GenericPlatform/GenericPlatformFile.h b/Engine/Source/Runtime/Core/Public/GenericPlatform/GenericPlatformFile.h
--- a/Engine/Source/Runtime/Core/Public/GenericPlatform/GenericPlatformFile.h
+++ b/Engine/Source/Runtime/Core/Public/GenericPlatform/GenericPlatformFile.h
@@ -934,6 +934,12 @@ public:
 	}
 };
 
+class IWrappedFileHandle : public IFileHandle
+{
+public:
+	virtual void SetLowerLevel(IFileHandle* LowerLevelHandle) = 0;
+};
+
 /**
 * Common base for physical platform File I/O Interface
 **/
diff --git a/Engine/Source/Runtime/Core/Public/GenericPlatform/GenericPlatformMisc.h b/Engine/Source/Runtime/Core/Public/GenericPlatform/GenericPlatformMisc.h
--- a/Engine/Source/Runtime/Core/Public/GenericPlatform/GenericPlatformMisc.h
+++ b/Engine/Source/Runtime/Core/Public/GenericPlatform/GenericPlatformMisc.h
@@ -1265,8 +1265,6 @@ public:
 	 */
 	static CORE_API IPlatformCompression* GetPlatformCompression();
 
-	static CORE_API IWrappedFileHandle* CreateCompressedPlatformFileHandle(IFileHandle* InLowerLevelHandle);
-
 	/**
 	 * Returns the platform specific interface for communication with processes running on the host pc.
 	 * 
diff --git a/Engine/Source/Runtime/Launch/Private/LaunchEngineLoop.cpp b/Engine/Source/Runtime/Launch/Private/LaunchEngineLoop.cpp
--- a/Engine/Source/Runtime/Launch/Private/LaunchEngineLoop.cpp
+++ b/Engine/Source/Runtime/Launch/Private/LaunchEngineLoop.cpp
@@ -1121,6 +1121,7 @@ bool LaunchCheckForFileOverride(const TCHAR* CmdLine, bool& OutFileOverrideFound
 		{
 			CurrentPlatformFile = NetworkPlatformFile;
 			FPlatformFileManager::Get().SetPlatformFile(*CurrentPlatformFile);
+			FPlatformFileManager::Get().SetPlatformPhysical(*CurrentPlatformFile);
 		}
 	}
 
@@ -1174,8 +1175,7 @@ bool LaunchCheckForFileOverride(const TCHAR* CmdLine, bool& OutFileOverrideFound
 		}
 	}
 #endif
-
-	if (!NetworkPlatformFile)
+	if (UE::IsUsingZenPakFileStreaming() || !NetworkPlatformFile)
 	{
 		IPlatformFile* PlatformFile = ConditionallyCreateFileWrapper(TEXT("PakFile"), CurrentPlatformFile, CmdLine);
 		if (PlatformFile)
@@ -1183,8 +1183,11 @@ bool LaunchCheckForFileOverride(const TCHAR* CmdLine, bool& OutFileOverrideFound
 			CurrentPlatformFile = PlatformFile;
 			FPlatformFileManager::Get().SetPlatformFile(*CurrentPlatformFile);
 		}
+	}
 
-		PlatformFile = ConditionallyCreateFileWrapper(TEXT("CachedReadFile"), CurrentPlatformFile, CmdLine);
+	if (!NetworkPlatformFile)
+	{
+		IPlatformFile* PlatformFile = ConditionallyCreateFileWrapper(TEXT("CachedReadFile"), CurrentPlatformFile, CmdLine);
 		if (PlatformFile)
 		{
 			CurrentPlatformFile = PlatformFile;
diff --git a/Engine/Source/Runtime/PakFile/Private/GenericPlatformIoDispatcher.cpp b/Engine/Source/Runtime/PakFile/Private/GenericPlatformIoDispatcher.cpp
--- a/Engine/Source/Runtime/PakFile/Private/GenericPlatformIoDispatcher.cpp
+++ b/Engine/Source/Runtime/PakFile/Private/GenericPlatformIoDispatcher.cpp
@@ -57,13 +57,22 @@ FGenericFileIoStoreImpl::~FGenericFileIoStoreImpl()
 
 bool FGenericFileIoStoreImpl::OpenContainer(const TCHAR* ContainerFilePath, uint64& ContainerFileHandle, uint64& ContainerFileSize)
 {
-	IPlatformFile& Ipf = IPlatformFile::GetPlatformPhysical();
-	int64 FileSize = Ipf.FileSize(ContainerFilePath);
+	IPlatformFile* Ipf = nullptr;
+	if (UE::IsUsingZenPakFileStreaming())
+	{
+		Ipf = &FPlatformFileManager::Get().GetPlatformFile();
+	}
+	else
+	{
+		Ipf = &IPlatformFile::GetPlatformPhysical();
+	}
+
+	const int64 FileSize = Ipf->FileSize(ContainerFilePath);
 	if (FileSize < 0)
 	{
 		return false;
 	}
-	IFileHandle* FileHandle = Ipf.OpenReadNoBuffering(ContainerFilePath);
+	IFileHandle* FileHandle = Ipf->OpenReadNoBuffering(ContainerFilePath);
 	if (!FileHandle)
 	{
 		return false;
diff --git a/Engine/Source/Runtime/PakFile/Private/IPlatformFilePak.cpp b/Engine/Source/Runtime/PakFile/Private/IPlatformFilePak.cpp
--- a/Engine/Source/Runtime/PakFile/Private/IPlatformFilePak.cpp
+++ b/Engine/Source/Runtime/PakFile/Private/IPlatformFilePak.cpp
@@ -3861,7 +3861,7 @@ public:
 		UE_LOG(LogPakFile, VeryVerbose, TEXT("FPakPlatformFile::OpenAsyncRead (FBypassPakAsyncReadFileHandle)[%016llX, %016llX) %s"), OffsetInPak, OffsetInPak + CompressedFileSize, Filename);
 		check(PakFileSize > 0 && OffsetInPak + CompressedFileSize <= PakFileSize && OffsetInPak >= 0);
 
-		LowerHandle = IPlatformFile::GetPlatformPhysical().OpenAsyncRead(*InPakFile->GetFilename());
+		LowerHandle = FPlatformFileManager::Get().GetPlatformPhysical().OpenAsyncRead(*InPakFile->GetFilename());
 	}
 	~FBypassPakAsyncReadFileHandle()
 	{
diff --git a/Engine/Source/Runtime/StorageServerClient/Private/StorageServerPlatformFile.cpp b/Engine/Source/Runtime/StorageServerClient/Private/StorageServerPlatformFile.cpp
--- a/Engine/Source/Runtime/StorageServerClient/Private/StorageServerPlatformFile.cpp
+++ b/Engine/Source/Runtime/StorageServerClient/Private/StorageServerPlatformFile.cpp
@@ -34,6 +34,19 @@ DEFINE_LOG_CATEGORY_STATIC(LogStorageServerPlatformFile, Log, All);
 #define EXCLUDE_NONSERVER_UE_EXTENSIONS 1	// Use .Build.cs file to disable this if the game relies on accessing loose files on the local filesystem
 #endif
 
+#if !defined(HAS_STORAGE_SERVER_COMPRESSED_FILE_HANDLE)
+#	define HAS_STORAGE_SERVER_COMPRESSED_FILE_HANDLE 0
+#endif
+
+#if HAS_STORAGE_SERVER_COMPRESSED_FILE_HANDLE
+	IWrappedFileHandle* CreateCompressedPlatformFileHandle(IFileHandle* InLowerLevelHandle);
+#else
+	IWrappedFileHandle* CreateCompressedPlatformFileHandle(IFileHandle* InLowerLevelHandle)
+	{
+		return nullptr;
+	}
+#endif // HAS_STORAGE_SERVER_COMPRESSED_FILE_HANDLE
+
 static FDateTime GAssumedImmutableTimeStamp = FDateTime::Now();
 
 FStorageServerFileSystemTOC::~FStorageServerFileSystemTOC()
@@ -163,6 +176,32 @@ bool FStorageServerFileSystemTOC::IterateDirectory(const FString& Path, TFunctio
 	return true;
 }
 
+bool FStorageServerFileSystemTOC::IterateDirectoryRecursively(const FString& Path, TFunctionRef<bool(const FIoChunkId&, const TCHAR*, int64)> Callback)
+{
+	UE_LOG(LogStorageServerPlatformFile, Verbose, TEXT("IterateDirectoryRecursively '%s'"), *Path);
+
+	FReadScopeLock _(TocLock);
+	FDirectory* Directory = Directories.FindRef(Path);
+	if (!Directory)
+	{
+		return false;
+	}
+	for (int32 FileIndex : Directory->Files)
+	{
+		const FFile& File = Files[FileIndex];
+		if (!Callback(File.FileChunkId, *File.FilePath, File.RawSize))
+		{
+			return false;
+		}
+	}
+	bool bFail = false;
+	for (const FString& ChildDirectoryPath : Directory->Directories)
+	{
+		bFail |= !IterateDirectoryRecursively(ChildDirectoryPath, Callback);
+	}
+
+	return !bFail;
+}
 
 #if COUNTERSTRACE_ENABLED
 	TRACE_DECLARE_ATOMIC_FLOAT_COUNTER(StorageServerCache_HitRatioBytes, TEXT("ZenClient/FileCache/HitRatio"));
@@ -537,6 +576,11 @@ public:
 
 FStorageServerPlatformFile::FStorageServerPlatformFile()
 {
+	if (UE::IsUsingZenPakFileStreaming())
+	{
+		ServerEngineDirView = FStringView(TEXT("Engine/"));
+		ServerProjectDirView = FStringView(TEXT(PREPROCESSOR_TO_STRING(UE_PROJECT_NAME)) TEXT("/"));
+	}
 }
 
 FStorageServerPlatformFile::~FStorageServerPlatformFile()
@@ -744,6 +788,11 @@ bool FStorageServerPlatformFile::Initialize(IPlatformFile* Inner, const TCHAR* C
 			UE_LOG(LogStorageServerPlatformFile, Display, TEXT("Using settings from command line: -ZenStoreBaseURI='%s'"), *BaseURI);
 		}
 
+		if (UE::IsUsingZenPakFileStreaming())
+		{
+			InitializeConnection();
+		}
+
 		return true;
 	}
 	return false;
@@ -924,7 +973,10 @@ FFileStatData FStorageServerPlatformFile::GetStatData(const TCHAR* FilenameOrDir
 
 IFileHandle* FStorageServerPlatformFile::InternalOpenFile(const FIoChunkId& FileChunkId, int64 RawSize, const TCHAR* LocalFilename)
 {
-	return new FStorageServerFileHandle(*this, FileChunkId, RawSize, LocalFilename);
+	IFileHandle* FileHandle = new FStorageServerFileHandle(*this, FileChunkId, RawSize, LocalFilename);
+	IWrappedFileHandle* FileDecompressor = CreateCompressedPlatformFileHandle(FileHandle);
+	
+	return FileDecompressor ? FileDecompressor : FileHandle;
 }
 
 IFileHandle* FStorageServerPlatformFile::OpenRead(const TCHAR* Filename, bool bAllowWrite)
@@ -965,6 +1017,29 @@ bool FStorageServerPlatformFile::IterateDirectory(const TCHAR* Directory, IPlatf
 	return bResult;
 }
 
+bool FStorageServerPlatformFile::IterateDirectoryRecursively(const TCHAR* Directory, IPlatformFile::FDirectoryVisitor& Visitor)
+{
+	TStringBuilder<1024> StorageServerDirectory;
+	bool bResult = false;
+	if (MakeStorageServerPath(Directory, StorageServerDirectory) && ServerToc.DirectoryExists(*StorageServerDirectory))
+	{
+		bResult |= ServerToc.IterateDirectoryRecursively(*StorageServerDirectory, [this, &Visitor](const FIoChunkId& FileChunkId, const TCHAR* FilenameOrDirectory, int64 RawSize)
+		{
+			TStringBuilder<1024> LocalPath;
+			bool bConverted = MakeLocalPath(FilenameOrDirectory, LocalPath);
+			check(bConverted);
+			const bool bDirectory = !FileChunkId.IsValid();
+			return Visitor.CallShouldVisitAndVisit(*LocalPath, bDirectory);
+		});
+	}
+	else
+	{
+		bResult |= LowerLevel->IterateDirectoryRecursively(Directory, Visitor);
+	}
+
+	return bResult;
+}
+
 bool FStorageServerPlatformFile::IterateDirectoryStat(const TCHAR* Directory, FDirectoryStatVisitor& Visitor)
 {
 	TStringBuilder<1024> StorageServerDirectory;
diff --git a/Engine/Source/Runtime/StorageServerClient/Private/StorageServerPlatformFile.h b/Engine/Source/Runtime/StorageServerClient/Private/StorageServerPlatformFile.h
--- a/Engine/Source/Runtime/StorageServerClient/Private/StorageServerPlatformFile.h
+++ b/Engine/Source/Runtime/StorageServerClient/Private/StorageServerPlatformFile.h
@@ -35,6 +35,7 @@ public:
 	int64 GetFileSize(const FString& Path);
 	bool GetFileData(const FString& Path, FIoChunkId& OutChunkId, int64& OutRawSize);
 	bool IterateDirectory(const FString& Path, TFunctionRef<bool(const FIoChunkId&, const TCHAR*, int64)> Callback);
+	bool IterateDirectoryRecursively(const FString& Path, TFunctionRef<bool(const FIoChunkId&, const TCHAR*, int64)> Callback);
 
 private:
 	struct FDirectory
@@ -93,6 +94,7 @@ public:
 	virtual bool DirectoryExists(const TCHAR* Directory) override;
 	virtual FFileStatData GetStatData(const TCHAR* FilenameOrDirectory) override;
 	virtual bool IterateDirectory(const TCHAR* Directory, FDirectoryVisitor& Visitor) override;
+	virtual bool IterateDirectoryRecursively(const TCHAR* Directory, FDirectoryVisitor& Visitor) override;
 	virtual bool IterateDirectoryStat(const TCHAR* Directory, FDirectoryStatVisitor& Visitor) override;
 	virtual IMappedFileHandle* OpenMapped(const TCHAR* Filename) override;
 	virtual FString GetFilenameOnDisk(const TCHAR* Filename) override;
