Fix trace.file when launching from editor to console.

diff --git a/Engine/Source/Runtime/StorageServerClient/Private/StorageServerPlatformFile.cpp b/Engine/Source/Runtime/StorageServerClient/Private/StorageServerPlatformFile.cpp
--- a/Engine/Source/Runtime/StorageServerClient/Private/StorageServerPlatformFile.cpp
+++ b/Engine/Source/Runtime/StorageServerClient/Private/StorageServerPlatformFile.cpp
@@ -1280,6 +1280,16 @@ FString FStorageServerPlatformFile::ConvertToAbsolutePathForExternalAppForRead(c
 	return IStorageServerPlatformFile::ConvertToAbsolutePathForExternalAppForRead(Filename);
 }
 
+FString FStorageServerPlatformFile::ConvertToAbsolutePathForExternalAppForWrite(const TCHAR* Filename)
+{
+	if (LowerLevel)
+	{
+		return LowerLevel->ConvertToAbsolutePathForExternalAppForWrite(Filename);
+	}
+
+	return IStorageServerPlatformFile::ConvertToAbsolutePathForExternalAppForWrite(Filename);
+}
+
 bool FStorageServerPlatformFile::IsNonServerFilenameAllowed(FStringView InFilename)
 {
 	bool bAllowed = true;
diff --git a/Engine/Source/Runtime/StorageServerClient/Private/StorageServerPlatformFile.h b/Engine/Source/Runtime/StorageServerClient/Private/StorageServerPlatformFile.h
--- a/Engine/Source/Runtime/StorageServerClient/Private/StorageServerPlatformFile.h
+++ b/Engine/Source/Runtime/StorageServerClient/Private/StorageServerPlatformFile.h
@@ -106,6 +106,7 @@ public:
 	virtual bool CreateDirectory(const TCHAR* Directory) override;
 	virtual bool DeleteDirectory(const TCHAR* Directory) override;
 	virtual FString ConvertToAbsolutePathForExternalAppForRead(const TCHAR* Filename) override;
+	virtual FString ConvertToAbsolutePathForExternalAppForWrite(const TCHAR* Filename) override;
 	virtual bool SendMessageToServer(const TCHAR* Message, IPlatformFile::IFileServerMessageHandler* Handler) override;
 
 	FStringView GetHostAddr() const override;
