Change IOS storage server streaming client to prefer the APIPA address used for ethernet-over-usb.

diff --git a/Engine/Source/Runtime/StorageServerClient/Private/StorageServerConnection.cpp b/Engine/Source/Runtime/StorageServerClient/Private/StorageServerConnection.cpp
--- a/Engine/Source/Runtime/StorageServerClient/Private/StorageServerConnection.cpp
+++ b/Engine/Source/Runtime/StorageServerClient/Private/StorageServerConnection.cpp
@@ -279,6 +279,19 @@ TArray<FString> FStorageServerConnection::SortHostAddressesByLocalSubnet(TArrayV
 				continue;
 			}
 		}
+#elif PLATFORM_IOS
+		if (Addr->GetProtocolType() == FNetworkProtocolTypes::IPv4)
+		{
+			// iOS and Mac have an APIPA address for the ethernet-over-usb connection 
+			// between the devices.  If we have an address that matches that pattern
+			// supplied by the Mac, we should prefer to use it above other reachable
+			// addresses.  Treating it as IPV4 loopback accomplishes that.
+			if (tempAddrStringSubnet.StartsWith("169.254."))
+			{
+				IPV4Loopback.Push(HostAddress);
+				continue;
+			}
+		}
 #endif
 
 		int32 LastDotPos = INDEX_NONE;
@@ -288,9 +301,13 @@ TArray<FString> FStorageServerConnection::SortHostAddressesByLocalSubnet(TArrayV
 		}
 
 		if (localAddrStringSubnet.Equals(tempAddrStringSubnet))
+		{
 			RegularAddresses.Insert(HostAddress, 0);
+		}
 		else
+		{
 			RegularAddresses.Push(HostAddress);
+		}
 	}
 
 	for (const FString& PlatformAddr : PlatformSocketAddress)
