Adding an experimental client-side cache for the Zen client, disabled by default.

 - implemented using IAS cache from the experimental IoStoreOnDemand module.
 - currently disabled because it does not appear to offer any speed improvements.

NB. IAS cache header moved to Internal

#rnx

diff --git a/Engine/Source/Runtime/Experimental/IoStore/OnDemand/Internal/IO/IasCache.h b/Engine/Source/Runtime/Experimental/IoStore/OnDemand/Internal/IO/IasCache.h
new file mode 100644
--- /dev/null
+++ b/Engine/Source/Runtime/Experimental/IoStore/OnDemand/Internal/IO/IasCache.h
@@ -0,0 +1,81 @@
+// Copyright Epic Games, Inc. All Rights Reserved.
+
+#pragma once
+
+#include "Containers/StringView.h"
+#include "IO/IoStatus.h"
+#include "Memory/MemoryFwd.h"
+#include "Tasks/Task.h"
+#include "Templates/UniquePtr.h"
+
+class FIoBuffer;
+class FIoReadOptions;
+struct FIoHash;
+
+namespace UE::IoStore
+{
+
+/** Cache for binary blobs with a 20 byte cache key. */
+class IIasCache
+{
+public:
+	virtual ~IIasCache() = default;
+
+	/** Deletes the IAS object, dropping all data persisted to disk and releasing
+	OS resources. As this also deletes the object thus any unique pointers should
+	be released prior to abandonment; TUniquePtr->Release()->Abandon(). Be sure
+	to cancel and collect any Get() tasks beforehand. */
+	virtual void Abandon() = 0;
+
+	/** Returns whether the specified cache key is present in the cache. */
+	virtual bool ContainsChunk(const FIoHash& Key) const = 0;
+
+	/** Get the chunk associated with the specified cache key. If the data is
+	 already in memory it is return in OutData. Otherwise the returned status
+	 indicates if the key can be materialized or if it does not exist */
+	virtual EIoErrorCode Get(const FIoHash& Key, FIoBuffer& OutData) = 0;
+
+	/** Materializes a cached items data from disk. The data is read into Dest
+	 so Dest must remain valid throughout. The result of the disk read is returned
+	 in Status (same lifetime needs as Dest). DoneEvent is triggered when the read
+	 succeeds, is not found, or if an IO error occurred. */
+	virtual void Materialize(const FIoHash& Key, FIoBuffer& Dest, EIoErrorCode& Status, UE::Tasks::FTaskEvent DoneEvent) = 0;
+
+	/** Cancels a previously request to Materialize(). Note that the materialize
+	 already can still complete as the read operation may already be in flight. If
+	 the cancel succeeds, the materialize's DoneEvent is not triggered. */
+	virtual void Cancel(FIoBuffer& GivenDest) = 0;
+
+	/** Insert a new chunk into the cache. */
+	virtual FIoStatus Put(const FIoHash& Key, FIoBuffer& Data) = 0;
+};
+
+struct FIasCacheConfig
+{
+	struct FRate
+	{
+		uint32	Allowance = 16 << 20;
+		uint32	Ops = 32;
+		uint32	Seconds = 60;
+	};
+
+	struct FDemand
+	{
+		uint8	Threshold = 30;
+		uint8	Boost = 60;
+		uint8	SuperBoost = 87;
+	};
+
+	FStringView Name = TEXT("ias");
+	uint64		DiskQuota = 512ull << 20;
+	uint32		MemoryQuota = 2 << 20;
+	uint32		JournalQuota = 4 << 20; // description in JournalCache.cpp
+	uint32		JournalMagic = 0; // can be used to invalidate cache
+	FRate		WriteRate;
+	FDemand		Demand;
+	bool		DropCache = false;
+};
+
+IOSTOREONDEMAND_API TUniquePtr<IIasCache> MakeIasCache(const TCHAR* RootDir, const FIasCacheConfig& Config);
+
+} // namespace UE::IoStore
diff --git a/Engine/Source/Runtime/Experimental/IoStore/OnDemand/Private/IasCache.h b/Engine/Source/Runtime/Experimental/IoStore/OnDemand/Private/IasCache.h
deleted file mode 100644
diff --git a/Engine/Source/Runtime/Experimental/IoStore/OnDemand/Private/IoStoreOnDemand.cpp b/Engine/Source/Runtime/Experimental/IoStore/OnDemand/Private/IoStoreOnDemand.cpp
--- a/Engine/Source/Runtime/Experimental/IoStore/OnDemand/Private/IoStoreOnDemand.cpp
+++ b/Engine/Source/Runtime/Experimental/IoStore/OnDemand/Private/IoStoreOnDemand.cpp
@@ -9,7 +9,7 @@
 #include "HAL/PlatformTime.h"
 #include "IO/IoChunkEncoding.h"
 #include "IO/PackageId.h"
-#include "IasCache.h"
+#include "IO/IasCache.h"
 #include "Misc/Base64.h"
 #include "Misc/CommandLine.h"
 #include "Misc/ConfigCacheIni.h"
diff --git a/Engine/Source/Runtime/Experimental/IoStore/OnDemand/Private/JournaledCache.cpp b/Engine/Source/Runtime/Experimental/IoStore/OnDemand/Private/JournaledCache.cpp
--- a/Engine/Source/Runtime/Experimental/IoStore/OnDemand/Private/JournaledCache.cpp
+++ b/Engine/Source/Runtime/Experimental/IoStore/OnDemand/Private/JournaledCache.cpp
@@ -1,6 +1,6 @@
 // Copyright Epic Games, Inc. All Rights Reserved.
 
-#include "IasCache.h"
+#include "IO/IasCache.h"
 #include "IO/IoStoreOnDemand.h"
 #include "Statistics.h"
 
diff --git a/Engine/Source/Runtime/Experimental/IoStore/OnDemand/Private/OnDemandIoDispatcherBackend.cpp b/Engine/Source/Runtime/Experimental/IoStore/OnDemand/Private/OnDemandIoDispatcherBackend.cpp
--- a/Engine/Source/Runtime/Experimental/IoStore/OnDemand/Private/OnDemandIoDispatcherBackend.cpp
+++ b/Engine/Source/Runtime/Experimental/IoStore/OnDemand/Private/OnDemandIoDispatcherBackend.cpp
@@ -21,7 +21,7 @@
 #include "IO/IoStatus.h"
 #include "IO/IoStore.h"
 #include "IO/IoStoreOnDemand.h"
-#include "IasCache.h"
+#include "IO/IasCache.h"
 #include "Logging/StructuredLog.h"
 #include "Math/NumericLimits.h"
 #include "Misc/ConfigCacheIni.h"
diff --git a/Engine/Source/Runtime/StorageServerClient/Private/StorageServerIoDispatcherBackend.cpp b/Engine/Source/Runtime/StorageServerClient/Private/StorageServerIoDispatcherBackend.cpp
--- a/Engine/Source/Runtime/StorageServerClient/Private/StorageServerIoDispatcherBackend.cpp
+++ b/Engine/Source/Runtime/StorageServerClient/Private/StorageServerIoDispatcherBackend.cpp
@@ -1,9 +1,10 @@
 // Copyright Epic Games, Inc. All Rights Reserved.
 
 #include "StorageServerIoDispatcherBackend.h"
-#include "Misc/QueuedThreadPool.h"
 #include "StorageServerConnection.h"
+#include "Misc/QueuedThreadPool.h"
 #include "Misc/ScopeLock.h"
+#include "Misc/CommandLine.h"
 #include "HAL/RunnableThread.h"
 #include "ProfilingDebugging/IoStoreTrace.h"
 
@@ -28,6 +29,13 @@ FStorageServerIoDispatcherBackend::FStorageServerIoDispatcherBackend(FStorageSer
 	, NewRequestEvent(FPlatformProcess::GetSynchEventFromPool(false))
 	, BatchCompletedEvent(FPlatformProcess::GetSynchEventFromPool(false))
 {
+#if WITH_STORAGESERVER_LOCAL_CACHE
+	bool bUseLocalCache = FParse::Param(FCommandLine::Get(), TEXT("ZenLocalCache"));
+	if (bUseLocalCache)
+	{
+		Cache = MakeShared<FStorageServerLocalCache>(Connection);
+	}
+#endif
 }
 
 FStorageServerIoDispatcherBackend::~FStorageServerIoDispatcherBackend()
@@ -156,6 +164,12 @@ bool FStorageServerIoDispatcherBackend::Resolve(FIoRequestImpl* Request)
 		{
 			Request->SetResult(Result.ValueOrDie());
 			TRACE_IOSTORE_BACKEND_REQUEST_COMPLETED(Request, Result.ValueOrDie().GetSize());
+#if WITH_STORAGESERVER_LOCAL_CACHE
+			if (Cache.IsValid())
+			{
+				Cache->OptionalCacheWrite(Request, Result.ValueOrDie());
+			}
+#endif
 		}
 		else
 		{
@@ -353,6 +367,12 @@ void FStorageServerIoDispatcherBackend::FBatch::DoThreadedWork()
 			Request->SetResult(Result.ValueOrDie());
 			TRACE_IOSTORE_BACKEND_REQUEST_COMPLETED(Request, Result.ValueOrDie().GetSize());
 
+#if WITH_STORAGESERVER_LOCAL_CACHE
+			if (Owner.Cache.IsValid())
+			{
+				Owner.Cache->OptionalCacheWrite(Request, Result.ValueOrDie());
+			}
+#endif
 		}
 		else
 		{
diff --git a/Engine/Source/Runtime/StorageServerClient/Private/StorageServerIoDispatcherBackend.h b/Engine/Source/Runtime/StorageServerClient/Private/StorageServerIoDispatcherBackend.h
--- a/Engine/Source/Runtime/StorageServerClient/Private/StorageServerIoDispatcherBackend.h
+++ b/Engine/Source/Runtime/StorageServerClient/Private/StorageServerIoDispatcherBackend.h
@@ -3,6 +3,7 @@
 #pragma once
 
 #include "IO/IoDispatcherBackend.h"
+#include "StorageServerLocalCache.h"
 #include "HAL/Runnable.h"
 #include "Misc/IQueuedWork.h"
 
@@ -98,6 +99,9 @@ private:
 	FIoRequestImpl* CompletedRequestsHead = nullptr;
 	FIoRequestImpl* CompletedRequestsTail = nullptr;
 	TAtomic<bool> bStopRequested{ false };
+#if WITH_STORAGESERVER_LOCAL_CACHE
+	TSharedPtr<FStorageServerLocalCache> Cache;
+#endif
 };
 
 #endif
diff --git a/Engine/Source/Runtime/StorageServerClient/Private/StorageServerLocalCache.cpp b/Engine/Source/Runtime/StorageServerClient/Private/StorageServerLocalCache.cpp
new file mode 100644
--- /dev/null
+++ b/Engine/Source/Runtime/StorageServerClient/Private/StorageServerLocalCache.cpp
@@ -0,0 +1,289 @@
+// Copyright Epic Games, Inc. All Rights Reserved.
+
+#include "StorageServerLocalCache.h"
+#include "StorageServerConnection.h"
+#include "Misc/ScopeLock.h"
+#include "Misc/Paths.h"
+#include "Misc/CoreDelegates.h"
+#include "Misc/CommandLine.h"
+#include "ProfilingDebugging/IoStoreTrace.h"
+#include "ProfilingDebugging/CountersTrace.h"
+#include "Modules/ModuleManager.h"
+
+#if !UE_BUILD_SHIPPING && WITH_STORAGESERVER_LOCAL_CACHE
+
+/******************************************************************************************************************************
+*                                                                                                                             * 
+*  Enable with -ZenLocalCache  + can use  -ZenDropLocalCache to clear the IAS cache on startup.								  *
+*                                                                                                                             * 
+*******************************************************************************************************************************/
+
+int32 GStorageServerLocalCacheMinDataSizeKB = 0;
+static FAutoConsoleVariableRef CVar_StorageServerLocalCacheMinDataSizeKB(
+	TEXT("s.StorageServerLocalCache.MinDataSizeKB"),
+	GStorageServerLocalCacheMinDataSizeKB,
+	TEXT("Minimum data size to store in the cache (KB)")
+);
+
+int32 GStorageServerLocalCacheMaxDataSizeKB = -1;
+static FAutoConsoleVariableRef CVar_StorageServerLocalCacheMaxDataSizeKB(
+	TEXT("s.StorageServerLocalCache.MaxDataSizeKB"),
+	GStorageServerLocalCacheMaxDataSizeKB,
+	TEXT("Maximum data size to store in the cache (KB). Set to -1 to try and cache everything")
+);
+
+namespace
+{
+	// this is a way to optionally ignore a certain proportion of items from the cache (based on ChunkId hash - not perfect)
+	int32 GCachePutMask = 0; //0 -> put all,  1 -> put 1/2, 3 -> put 1/4, 7 -> put 1/8, 15 -> put 1/16
+
+	// disable all cache writes (for testing/debugging)
+	bool bCacheIsReadOnly = false;
+}
+
+TRACE_DECLARE_INT_COUNTER(ZenCacheHitBytes,         TEXT("ZenClient/LocalCache/MemoryHitBytes"));
+TRACE_DECLARE_INT_COUNTER(ZenCacheMaterializeBytes, TEXT("ZenClient/LocalCache/DiskHitBytes"));
+TRACE_DECLARE_INT_COUNTER(ZenCacheMissBytes,        TEXT("ZenClient/LocalCache/MissBytes"));
+TRACE_DECLARE_INT_COUNTER(ZenCachePutBytes,         TEXT("ZenClient/LocalCache/PutBytes"));
+
+/******************************************************************************************************************************/
+
+
+
+
+FStorageServerLocalCache::FStorageServerLocalCache(FStorageServerConnection& InConnection) 
+	: Connection(InConnection)
+{
+	if (FModuleManager::Get().ModuleExists(TEXT("IoStoreOnDemand")) )
+	{
+		FCoreDelegates::OnInit.AddLambda([this]()
+		{
+			FString CacheDir = FPaths::ProjectPersistentDownloadDir();
+
+			if (!CacheDir.IsEmpty())
+			{
+				TRACE_CPUPROFILER_EVENT_SCOPE(StorageServerIoDispatcherInitCache);
+
+				UE::IoStore::FIasCacheConfig Config;
+				Config.Name = TEXT("zencache");
+				Config.DropCache = FParse::Param(FCommandLine::Get(), TEXT("ZenDropLocalCache"));
+				Cache = UE::IoStore::MakeIasCache(*CacheDir, Config);
+
+				TRACE_IOSTORE_BACKEND_NAME(this, TEXT("StorageServerCache") );
+				
+				FIoDispatcher::Get().Mount(AsShared(), 1); //must be higher priority than FStorageServerIoDispatcherBackend
+			}
+		});
+	}
+}
+
+void FStorageServerLocalCache::Initialize(TSharedRef<const FIoDispatcherBackendContext> Context) 
+{
+	BackendContext = Context; 
+}
+
+
+FIoHash FStorageServerLocalCache::GetRequestCacheKey(FIoRequestImpl* Request) const
+{
+	FIoOffsetAndLength Range(Request->Options.GetOffset(), Request->Options.GetSize());
+
+	FIoHashBuilder HashBuilder;
+	HashBuilder.Update(Request->ChunkId.GetData(), Request->ChunkId.GetSize());
+	HashBuilder.Update(&Range, sizeof(FIoOffsetAndLength));
+	return HashBuilder.Finalize();
+}
+
+bool FStorageServerLocalCache::ResolveFromCache(FIoRequestImpl* Request)
+{
+	// bail out if there's no cache
+	if (!Cache.IsValid())
+	{
+		return false;
+	}
+
+	// bail out if we know isn't something that would be in the cache
+	uint32 ChunkHash = GetTypeHash(Request->ChunkId);
+	if (GCachePutMask > 0 && ((ChunkHash & GCachePutMask) != GCachePutMask))
+	{
+		return false;
+	}
+
+	// see if the chunk is in the cache
+	FIoHash Key = GetRequestCacheKey(Request);
+	if (!Cache->ContainsChunk(Key))
+	{
+		return false;
+	}
+
+	// helper lambda to set the result buffer
+	auto SetRequestResult = [Request, this](FIoBuffer& InBuffer)
+	{
+		if (Request->Options.GetTargetVa())
+		{
+			FIoBuffer VaBuffer(FIoBuffer::Wrap, Request->Options.GetTargetVa(), Request->Options.GetSize());
+			FMemory::Memcpy(VaBuffer.GetData(), InBuffer.GetData(), InBuffer.GetSize());
+			Request->SetResult(VaBuffer);
+
+		}
+		else
+		{
+			Request->SetResult(InBuffer);
+		}
+	};
+
+	// start processing - first, check the memory cache
+	EIoErrorCode ErrorCode = EIoErrorCode::Unknown;
+	{
+		FIoBuffer Buffer;
+		ErrorCode = Cache->Get(Key, Buffer);
+		if (ErrorCode == EIoErrorCode::Ok)
+		{
+			TRACE_IOSTORE_BACKEND_REQUEST_STARTED(Request, this);
+			SetRequestResult(Buffer);
+			CompletedRequests.Enqueue(Request);
+			BackendContext->WakeUpDispatcherThreadDelegate.Execute();
+			TRACE_COUNTER_ADD(ZenCacheHitBytes, Buffer.GetSize());
+			TRACE_IOSTORE_BACKEND_REQUEST_COMPLETED(Request, Buffer.GetSize());
+			return true;
+		}
+	}
+
+	// FileNotOpen means it is cached on disk - materialize it now
+	if (ErrorCode == EIoErrorCode::FileNotOpen)
+	{
+		TRACE_IOSTORE_BACKEND_REQUEST_STARTED(Request, this);
+
+		struct FPendingMaterialize
+		{
+			EIoErrorCode ErrorCode;
+			FIoBuffer Buffer;
+		};
+		FPendingMaterialize* PendingMaterialize = new FPendingMaterialize();
+		
+		UE::Tasks::FTaskEvent MaterializeDone(TEXT("ZenIasCacheMaterializeDone"));
+		UE::Tasks::FTask TaskHandle = UE::Tasks::Launch(UE_SOURCE_LOCATION, [this, Request, PendingMaterialize, SetRequestResult = MoveTemp(SetRequestResult)]
+		{
+			// complete the request
+			if (PendingMaterialize->ErrorCode == EIoErrorCode::Ok)
+			{
+				SetRequestResult(PendingMaterialize->Buffer);
+				TRACE_COUNTER_ADD(ZenCacheMaterializeBytes, PendingMaterialize->Buffer.GetSize());
+				TRACE_IOSTORE_BACKEND_REQUEST_COMPLETED(Request, PendingMaterialize->Buffer.GetSize());
+			}
+			else
+			{
+				Request->SetFailed();
+				TRACE_IOSTORE_BACKEND_REQUEST_FAILED(Request);
+			}
+			CompletedRequests.Enqueue(Request);
+			BackendContext->WakeUpDispatcherThreadDelegate.Execute();
+
+			//clean up
+			delete PendingMaterialize;
+		}, MaterializeDone, UE::Tasks::ETaskPriority::Normal);
+
+		Cache->Materialize(Key, PendingMaterialize->Buffer, PendingMaterialize->ErrorCode, MaterializeDone);
+		return true;
+	}
+
+	// can't process this request
+	return false;
+}
+
+void FStorageServerLocalCache::OptionalCacheWrite(FIoRequestImpl* Request, const FIoBuffer& Buffer)
+{
+	if (Cache.IsValid())
+	{
+		uint64 BufferSize = Buffer.GetSize();
+		TRACE_COUNTER_ADD(ZenCacheMissBytes, BufferSize);
+
+		// bail out if it's read only (for testing only!)
+		if (bCacheIsReadOnly)
+		{
+			return;
+		}
+
+
+		// bail out if this isn't something we're going to be caching
+		uint32 ChunkHash = GetTypeHash(Request->ChunkId);
+		if (GCachePutMask > 0 && ((ChunkHash & GCachePutMask) != GCachePutMask))
+		{
+			return;
+		}
+
+		// bail out if it's too small
+		if (BufferSize < GStorageServerLocalCacheMinDataSizeKB*1024)
+		{
+			return;
+		}
+
+		// bail out if it's too large
+		if (GStorageServerLocalCacheMaxDataSizeKB > 0 && BufferSize > GStorageServerLocalCacheMaxDataSizeKB*1024)
+		{
+			return;
+		}
+
+		// take a copy of the buffer and store it in the cache
+		FIoBuffer PutBuffer(Buffer.GetSize());
+		FMemory::Memcpy(PutBuffer.GetData(), Buffer.GetData(), Buffer.GetSize());
+
+		TRACE_COUNTER_ADD(ZenCachePutBytes, BufferSize);
+		Cache->Put(GetRequestCacheKey(Request), PutBuffer);
+	}
+}
+
+void FStorageServerLocalCache::ResolveIoRequests(FIoRequestList Requests, FIoRequestList& OutUnresolved)
+{
+	if (Cache.IsValid() && BackendContext->bIsMultiThreaded)
+	{
+		while (FIoRequestImpl* Request = Requests.PopHead())
+		{
+			if (ResolveFromCache(Request) == false)
+			{
+				OutUnresolved.AddTail(Request);
+			}
+		}
+	}
+}
+
+FIoRequestImpl* FStorageServerLocalCache::GetCompletedIoRequests()
+{
+	FIoRequestImpl* Result = nullptr;
+
+	FIoRequestImpl* Item = nullptr;
+	while (CompletedRequests.Dequeue(Item))
+	{
+		Item->NextRequest = Result;
+		Result = Item;
+	}
+
+	return Result;
+}
+
+bool FStorageServerLocalCache::DoesChunkExist(const FIoChunkId& ChunkId) const
+{
+	return GetSizeForChunk(ChunkId).IsOk();
+}
+
+TIoStatusOr<uint64> FStorageServerLocalCache::GetSizeForChunk(const FIoChunkId& ChunkId) const
+{
+	TRACE_CPUPROFILER_EVENT_SCOPE(StorageServerIoDispatcherCacheGetSizeForChunk);
+	int64 ChunkSize = Connection.ChunkSizeRequest(ChunkId);
+	if (ChunkSize >= 0)
+	{
+		return ChunkSize;
+	}
+	else
+	{
+		return FIoStatus(EIoErrorCode::NotFound);
+	}
+}
+
+const TCHAR* FStorageServerLocalCache::GetName() const
+{
+	return TEXT("StorageServerLocalCache");
+}
+
+
+#endif //!UE_BUILD_SHIPPING
+
diff --git a/Engine/Source/Runtime/StorageServerClient/Private/StorageServerLocalCache.h b/Engine/Source/Runtime/StorageServerClient/Private/StorageServerLocalCache.h
new file mode 100644
--- /dev/null
+++ b/Engine/Source/Runtime/StorageServerClient/Private/StorageServerLocalCache.h
@@ -0,0 +1,50 @@
+// Copyright Epic Games, Inc. All Rights Reserved.
+
+#pragma once
+
+#include "IO/IoDispatcherBackend.h"
+
+#ifndef WITH_STORAGESERVER_LOCAL_CACHE
+	#define WITH_STORAGESERVER_LOCAL_CACHE 0 // disabled by default
+#endif
+
+#if !UE_BUILD_SHIPPING && WITH_STORAGESERVER_LOCAL_CACHE
+#include "IO/IasCache.h"
+#include "Containers/MpscQueue.h"
+
+class FStorageServerConnection;
+
+class FStorageServerLocalCache : public IIoDispatcherBackend, public TSharedFromThis<FStorageServerLocalCache>
+{
+public:
+	FStorageServerLocalCache(FStorageServerConnection& InConnection);
+	virtual ~FStorageServerLocalCache() = default;
+
+	bool ResolveFromCache(FIoRequestImpl* Request);
+	void OptionalCacheWrite(FIoRequestImpl* Request, const FIoBuffer& Buffer);
+
+private:
+	FIoHash GetRequestCacheKey(FIoRequestImpl* Request) const;
+
+	FStorageServerConnection& Connection;
+	TUniquePtr<UE::IoStore::IIasCache> Cache;
+	TMpscQueue<FIoRequestImpl*> CompletedRequests;
+	TSharedPtr<const FIoDispatcherBackendContext> BackendContext;
+
+
+	// IIoDispatcherBackend implementation
+	virtual void Initialize(TSharedRef<const FIoDispatcherBackendContext> Context) override;
+	virtual void ResolveIoRequests(FIoRequestList Requests, FIoRequestList& OutUnresolved)  override;
+	virtual FIoRequestImpl* GetCompletedIoRequests()  override;
+	virtual void CancelIoRequest(FIoRequestImpl* Request)  override {}
+	virtual void UpdatePriorityForIoRequest(FIoRequestImpl* Request)  override {}
+	virtual bool DoesChunkExist(const FIoChunkId& ChunkId) const  override;
+	virtual TIoStatusOr<uint64> GetSizeForChunk(const FIoChunkId& ChunkId) const override;
+	virtual TIoStatusOr<FIoMappedRegion> OpenMapped(const FIoChunkId& ChunkId, const FIoReadOptions& Options) override
+	{ 
+		return FIoStatus(EIoErrorCode::NotFound); 
+	}
+	virtual const TCHAR* GetName() const override;
+};
+
+#endif
diff --git a/Engine/Source/Runtime/StorageServerClient/StorageServerClient.Build.cs b/Engine/Source/Runtime/StorageServerClient/StorageServerClient.Build.cs
--- a/Engine/Source/Runtime/StorageServerClient/StorageServerClient.Build.cs
+++ b/Engine/Source/Runtime/StorageServerClient/StorageServerClient.Build.cs
@@ -12,7 +12,8 @@ public class StorageServerClient : ModuleRules
 				"CoreUObject",
 				"Sockets",
 				"CookOnTheFly",
-				"Json"
+				"Json",
+				"IoStoreOnDemand" // for IAS cache
 			}
 		);
 	}
