Don't try to create IoStoreOnDemand or IAS cache if there is no persistent storage available

#rnx

diff --git a/Engine/Source/Runtime/Experimental/IoStore/OnDemand/Private/OnDemandIoStore.cpp b/Engine/Source/Runtime/Experimental/IoStore/OnDemand/Private/OnDemandIoStore.cpp
--- a/Engine/Source/Runtime/Experimental/IoStore/OnDemand/Private/OnDemandIoStore.cpp
+++ b/Engine/Source/Runtime/Experimental/IoStore/OnDemand/Private/OnDemandIoStore.cpp
@@ -123,7 +123,11 @@ static FString GetInstallCacheDirectory()
 
 	return FPaths::ProjectPersistentDownloadDir() / TEXT("IoStore") / DirName;
 }
-
+///////////////////////////////////////////////////////////////////////////////
+static bool HasInstallCacheDirectory()
+{
+	return FPaths::HasProjectPersistentDownloadDir();
+}
 ///////////////////////////////////////////////////////////////////////////////
 static void SplitHostUrl(const FStringView& Url, FStringView& OutHost, FStringView& OutRemainder)
 {
@@ -793,6 +797,12 @@ FIoStatus FOnDemandIoStore::InitializeOnDemandInstallCache()
 #endif
 	if (bUseInstallCache)
 	{
+		if (!Private::HasInstallCacheDirectory())
+		{
+			UE_LOG(LogIoStoreOnDemand, Error, TEXT("Project does not have persistent storage configured"));
+			return FIoStatusBuilder(EIoErrorCode::InvalidParameter) << TEXT("Failed to initialize install cache");
+		}
+
 		FOnDemandInstallCacheConfig CacheConfig;
 
 		const int64 DiskQuota = []
diff --git a/Engine/Source/Runtime/StorageServerClient/Private/StorageServerLocalCache.cpp b/Engine/Source/Runtime/StorageServerClient/Private/StorageServerLocalCache.cpp
--- a/Engine/Source/Runtime/StorageServerClient/Private/StorageServerLocalCache.cpp
+++ b/Engine/Source/Runtime/StorageServerClient/Private/StorageServerLocalCache.cpp
@@ -58,21 +58,23 @@ FStorageServerLocalCache::FStorageServerLocalCache(FStorageServerConnection& InC
 	{
 		FCoreDelegates::OnInit.AddLambda([this]()
 		{
-			FString CacheDir = FPaths::ProjectPersistentDownloadDir();
-
-			if (!CacheDir.IsEmpty())
-			{
-				TRACE_CPUPROFILER_EVENT_SCOPE(StorageServerIoDispatcherInitCache);
-
-				UE::IoStore::FIasCacheConfig Config;
-				Config.Name = TEXT("zencache");
-				Config.DropCache = FParse::Param(FCommandLine::Get(), TEXT("ZenDropLocalCache"));
-				Cache = UE::IoStore::MakeIasCache(*CacheDir, Config);
-
-				TRACE_IOSTORE_BACKEND_NAME(this, TEXT("StorageServerCache") );
+				if (FPaths::HasProjectPersistentDownloadDir())
+				{
+					FString CacheDir = FPaths::ProjectPersistentDownloadDir();
+					if (!CacheDir.IsEmpty())
+					{
+						TRACE_CPUPROFILER_EVENT_SCOPE(StorageServerIoDispatcherInitCache);
+
+						UE::IoStore::FIasCacheConfig Config;
+						Config.Name = TEXT("zencache");
+						Config.DropCache = FParse::Param(FCommandLine::Get(), TEXT("ZenDropLocalCache"));
+						Cache = UE::IoStore::MakeIasCache(*CacheDir, Config);
+
+						TRACE_IOSTORE_BACKEND_NAME(this, TEXT("StorageServerCache") );
 				
-				FIoDispatcher::Get().Mount(AsShared(), 1); //must be higher priority than FStorageServerIoDispatcherBackend
-			}
+						FIoDispatcher::Get().Mount(AsShared(), 1); //must be higher priority than FStorageServerIoDispatcherBackend
+					}
+				}
 		});
 	}
 }
