Remove unnecessary include of "IO/IoChunkId.h"

#rnx

diff --git a/Engine/Source/Runtime/StorageServerClient/Private/StorageServerConnection.cpp b/Engine/Source/Runtime/StorageServerClient/Private/StorageServerConnection.cpp
--- a/Engine/Source/Runtime/StorageServerClient/Private/StorageServerConnection.cpp
+++ b/Engine/Source/Runtime/StorageServerClient/Private/StorageServerConnection.cpp
@@ -2,6 +2,7 @@
 
 #include "StorageServerConnection.h"
 
+#include "IO/IoChunkId.h"
 #include "IO/IoDispatcher.h"
 #include "IPAddress.h"
 #include "Misc/App.h"
diff --git a/Engine/Source/Runtime/StorageServerClient/Private/StorageServerConnection.h b/Engine/Source/Runtime/StorageServerClient/Private/StorageServerConnection.h
--- a/Engine/Source/Runtime/StorageServerClient/Private/StorageServerConnection.h
+++ b/Engine/Source/Runtime/StorageServerClient/Private/StorageServerConnection.h
@@ -12,7 +12,6 @@
 #include "SocketTypes.h"
 #include "HAL/PlatformTime.h"
 #include "StorageServerHttpClient.h"
-#include "IO/IoChunkId.h"
 #include "IStorageServerPlatformFile.h"
 
 #if !UE_BUILD_SHIPPING
@@ -20,6 +19,7 @@
 DECLARE_LOG_CATEGORY_EXTERN(LogStorageServerConnection, Log, All);
 
 struct FPackageStoreEntryResource;
+class FIoChunkId;
 
 class FStorageServerConnection
 {
