Storage server package store initializes asyncronously to drastically improve engine startup time when using Zen client.
- package store isn't needed until ProcessNewlyLoadedUObjects so there should be plenty of time to load the package store

#rnx

diff --git a/Engine/Source/Runtime/StorageServerClient/Private/StorageServerPackageStore.cpp b/Engine/Source/Runtime/StorageServerClient/Private/StorageServerPackageStore.cpp
--- a/Engine/Source/Runtime/StorageServerClient/Private/StorageServerPackageStore.cpp
+++ b/Engine/Source/Runtime/StorageServerClient/Private/StorageServerPackageStore.cpp
@@ -4,22 +4,27 @@
 #include "IO/IoDispatcher.h"
 #include "StorageServerConnection.h"
 #include "Serialization/MemoryReader.h"
+#include "HAL/RunnableThread.h"
+#include "HAL/Event.h"
 
 #if !UE_BUILD_SHIPPING
 
 FStorageServerPackageStoreBackend::FStorageServerPackageStoreBackend(FStorageServerConnection& Connection)
 {
-	Connection.PackageStoreRequest([this](FPackageStoreEntryResource&& PackageStoreEntryResource)
-		{
-			FStoreEntry& LocalEntry = StoreEntriesMap.FindOrAdd(PackageStoreEntryResource.GetPackageId());
-			LocalEntry.ImportedPackages = MoveTemp(PackageStoreEntryResource.ImportedPackageIds);
-			LocalEntry.ShaderMapHashes = MoveTemp(PackageStoreEntryResource.ShaderMapHashes);
-		});
+	AsyncInit = MakeShared<FAsyncInitRunnable>(*this, Connection);
 }
 
 EPackageStoreEntryStatus FStorageServerPackageStoreBackend::GetPackageStoreEntry(FPackageId PackageId, FName PackageName,
 	FPackageStoreEntry& OutPackageStoreEntry)
 {
+	// nb. the async init is highly likely to have finished by the time the first entry is requested
+	if (AsyncInit.IsValid())
+	{
+		TRACE_CPUPROFILER_EVENT_SCOPE(StorageServerPackageStoreWaitForInit);
+		AsyncInit->WaitForCompletion();
+		AsyncInit.Reset();
+	}
+
 	const FStoreEntry* FindEntry = StoreEntriesMap.Find(PackageId);
 	if (FindEntry)
 	{
@@ -33,4 +38,43 @@ EPackageStoreEntryStatus FStorageServerPackageStoreBackend::GetPackageStoreEntry
 	}
 }
 
+
+
+
+FStorageServerPackageStoreBackend::FAsyncInitRunnable::FAsyncInitRunnable(FStorageServerPackageStoreBackend& InOwner, FStorageServerConnection& InConnection)
+	: Owner(InOwner)
+	, Connection(InConnection)
+	, IsCompleted(FPlatformProcess::GetSynchEventFromPool(true))
+{
+	if (FRunnableThread::Create( this, TEXT("StorageServerPackageStoreInit"), 0, EThreadPriority::TPri_Normal) == nullptr)
+	{
+		IsCompleted->Trigger();
+	}
+}
+	
+FStorageServerPackageStoreBackend::FAsyncInitRunnable::~FAsyncInitRunnable()
+{
+	IsCompleted->Wait();
+	FPlatformProcess::ReturnSynchEventToPool(IsCompleted);
+}
+
+void FStorageServerPackageStoreBackend::FAsyncInitRunnable::WaitForCompletion() const
+{
+	IsCompleted->Wait();
+}
+
+uint32 FStorageServerPackageStoreBackend::FAsyncInitRunnable::Run()
+{
+	TRACE_CPUPROFILER_EVENT_SCOPE(StorageServerPackageStoreRequest);
+	Connection.PackageStoreRequest([this](FPackageStoreEntryResource&& PackageStoreEntryResource)
+		{
+			FStoreEntry& LocalEntry = Owner.StoreEntriesMap.FindOrAdd(PackageStoreEntryResource.GetPackageId());
+			LocalEntry.ImportedPackages = MoveTemp(PackageStoreEntryResource.ImportedPackageIds);
+			LocalEntry.ShaderMapHashes = MoveTemp(PackageStoreEntryResource.ShaderMapHashes);
+		});
+
+	IsCompleted->Trigger();
+	return 0;
+}
+
 #endif
diff --git a/Engine/Source/Runtime/StorageServerClient/Private/StorageServerPackageStore.h b/Engine/Source/Runtime/StorageServerClient/Private/StorageServerPackageStore.h
--- a/Engine/Source/Runtime/StorageServerClient/Private/StorageServerPackageStore.h
+++ b/Engine/Source/Runtime/StorageServerClient/Private/StorageServerPackageStore.h
@@ -3,6 +3,8 @@
 #pragma once
 
 #include "IO/PackageStore.h"
+#include "Templates/SharedPointer.h"
+#include "HAL/Runnable.h"
 
 #if !UE_BUILD_SHIPPING
 
@@ -43,6 +45,22 @@ private:
 		TArray<FSHAHash> ShaderMapHashes;
 	};
 	TMap<FPackageId, FStoreEntry> StoreEntriesMap;
+
+	class FAsyncInitRunnable : public FRunnable
+	{
+	public:
+		FAsyncInitRunnable(FStorageServerPackageStoreBackend& InOwner, FStorageServerConnection& InConnection);
+		virtual ~FAsyncInitRunnable();
+		void WaitForCompletion() const;
+
+	private:
+		virtual uint32 Run() override;
+
+		FStorageServerPackageStoreBackend& Owner;
+		FStorageServerConnection& Connection;
+		class FEvent* IsCompleted;
+	};
+	TSharedPtr<FAsyncInitRunnable, ESPMode::ThreadSafe> AsyncInit;
 };
 
 #endif
\ No newline at end of file
