Bring back compression to Zen requests /entries and /files

#rnx

diff --git a/Engine/Source/Runtime/StorageServerClient/Private/StorageServerConnection.cpp b/Engine/Source/Runtime/StorageServerClient/Private/StorageServerConnection.cpp
--- a/Engine/Source/Runtime/StorageServerClient/Private/StorageServerConnection.cpp
+++ b/Engine/Source/Runtime/StorageServerClient/Private/StorageServerConnection.cpp
@@ -29,6 +29,28 @@ DEFINE_LOG_CATEGORY(LogStorageServerConnection);
 TRACE_DECLARE_INT_COUNTER(ZenHttpClientSerializedBytes, TEXT("ZenClient/SerializedBytes (compressed)"));
 TRACE_DECLARE_INT_COUNTER(ZenHttpClientThroughputBytes, TEXT("ZenClient/ThroughputBytes (decompressed)"));
 
+static FCbObject ProcessResponse(FIoBuffer IoBuffer, EStorageServerContentType ContentType)
+{
+	FCbObject ResponseObj;
+	if (ContentType == EStorageServerContentType::CbObject)
+	{
+		FMemoryReaderView Reader(IoBuffer.GetView());
+		ResponseObj = LoadCompactBinary(Reader).AsObject();
+	}
+	else if (ContentType == EStorageServerContentType::CompressedBinary)
+	{
+		FCompressedBuffer Compressed = FCompressedBuffer::FromCompressed(FSharedBuffer::MakeView(IoBuffer.GetData(), IoBuffer.GetSize()));
+		FIoBuffer Decompressed(Compressed.GetRawSize());
+		if (FCompressedBufferReader(Compressed).TryDecompressTo(Decompressed.GetMutableView(), 0))
+		{
+			FBufferReader DecompressedAr(Decompressed.GetData(), Decompressed.GetSize(), false);
+			ResponseObj = LoadCompactBinary(DecompressedAr).AsObject();
+		}
+	}
+
+	return ResponseObj;
+}
+
 bool FStorageServerConnection::Initialize(TArrayView<const FString> HostAddresses, const int32 Port, const FAnsiStringView& InBaseURI)
 {
 	BaseURI = InBaseURI;
@@ -54,12 +76,12 @@ void FStorageServerConnection::PackageStoreRequest(TFunctionRef<void(FPackageSto
 	TAnsiStringBuilder<256> ResourceBuilder;
 	ResourceBuilder.Append(BaseURI).Append("/entries?fieldfilter=packagestoreentry");
 
-	IStorageServerHttpClient::FResult ResultTuple = HttpClient->RequestSync(*ResourceBuilder, EStorageServerContentType::CbObject);
+	IStorageServerHttpClient::FResult ResultTuple = HttpClient->RequestSync(*ResourceBuilder, EStorageServerContentType::CompressedBinary);
 	TIoStatusOr<FIoBuffer> Result = ResultTuple.Get<0>();
 	if (Result.IsOk())
 	{
-		FMemoryReaderView Reader(Result.ValueOrDie().GetView());
-		FCbObject ResponseObj = LoadCompactBinary(Reader).AsObject();
+		EStorageServerContentType ContentType = ResultTuple.Get<1>();
+		FCbObject ResponseObj = ProcessResponse(Result.ValueOrDie(), ContentType);
 		{
 			TRACE_CPUPROFILER_EVENT_SCOPE(StorageServerPackageStoreRequestParseEntries);
 			for (FCbField& OplogEntry : ResponseObj["entries"].AsArray())
@@ -83,12 +105,12 @@ void FStorageServerConnection::FileManifestRequest(TFunctionRef<void(FIoChunkId
 	TAnsiStringBuilder<256> ResourceBuilder;
 	ResourceBuilder.Append(BaseURI).Append("/files?fieldnames=id,clientpath,rawsize");
 
-	IStorageServerHttpClient::FResult ResultTuple = HttpClient->RequestSync(*ResourceBuilder, EStorageServerContentType::CbObject);
+	IStorageServerHttpClient::FResult ResultTuple = HttpClient->RequestSync(*ResourceBuilder, EStorageServerContentType::CompressedBinary);
 	TIoStatusOr<FIoBuffer> Result = ResultTuple.Get<0>();
 	if (Result.IsOk())
 	{
-		FMemoryReaderView Reader(Result.ValueOrDie().GetView());
-		FCbObject ResponseObj = LoadCompactBinary(Reader).AsObject();
+		EStorageServerContentType ContentType = ResultTuple.Get<1>();
+		FCbObject ResponseObj = ProcessResponse(Result.ValueOrDie(), ContentType);
 		for (FCbField& FileArrayEntry : ResponseObj["files"].AsArray())
 		{
 			FCbObject Entry = FileArrayEntry.AsObject();
