Early precache for .ini and .uplugin files to improve engine startup time when using Zen client.

#rnx

diff --git a/Engine/Source/Runtime/StorageServerClient/Private/StorageServerPlatformFile.cpp b/Engine/Source/Runtime/StorageServerClient/Private/StorageServerPlatformFile.cpp
--- a/Engine/Source/Runtime/StorageServerClient/Private/StorageServerPlatformFile.cpp
+++ b/Engine/Source/Runtime/StorageServerClient/Private/StorageServerPlatformFile.cpp
@@ -5,6 +5,9 @@
 #include "CookOnTheFly.h"
 #include "CookOnTheFlyPackageStore.h"
 #include "HAL/FileManagerGeneric.h"
+#include "HAL/Runnable.h"
+#include "HAL/RunnableThread.h"
+#include "HAL/CriticalSection.h"
 #include "Misc/App.h"
 #include "Misc/CommandLine.h"
 #include "Misc/CoreDelegates.h"
@@ -24,6 +27,7 @@
 #include "StorageServerIoDispatcherBackend.h"
 #include "StorageServerPackageStore.h"
 #include "Containers/LruCache.h"
+#include "Containers/SpscQueue.h"
 #include <atomic>
 #include "Async/MappedFileHandle.h"
 
@@ -35,6 +39,12 @@ DEFINE_LOG_CATEGORY_STATIC(LogStorageServerPlatformFile, Log, All);
 #define EXCLUDE_NONSERVER_UE_EXTENSIONS 1	// Use .Build.cs file to disable this if the game relies on accessing loose files on the local filesystem
 #endif
 
+
+#if !defined(WITH_STORAGE_SERVER_STARTUP_FILE_CACHE)
+#define WITH_STORAGE_SERVER_STARTUP_FILE_CACHE 1
+#endif
+
+
 #if !defined(HAS_STORAGE_SERVER_COMPRESSED_FILE_HANDLE)
 #	define HAS_STORAGE_SERVER_COMPRESSED_FILE_HANDLE 0
 #endif
@@ -204,6 +214,192 @@ bool FStorageServerFileSystemTOC::IterateDirectoryRecursively(const FString& Pat
 	return !bFail;
 }
 
+
+
+
+
+#if WITH_STORAGE_SERVER_STARTUP_FILE_CACHE
+
+#define STORAGE_SERVER_START_CACHE_REPORT 0 // log out startup cache usage at the end of engine startup
+
+class FStorageServerEngineStartupPrecache : public FRunnable
+{
+public:
+
+	// largest file size that will be stored in the cache
+	static const uint32 MaxFileSize = 16 * 1024;
+
+	// largest size that the cache is allowed to grow to
+	static const uint64 MaxCacheSize = 10 * 1024 * 1024;
+
+
+	FStorageServerEngineStartupPrecache(FStorageServerConnection& InConnection)
+		: Connection(InConnection)
+		, CacheSize(0)
+		, HasWork(FPlatformProcess::GetSynchEventFromPool())
+		, IsCompleted(FPlatformProcess::GetSynchEventFromPool())
+		, bExitWhenComplete(false)
+		, bExitImmediately(false)
+	{
+		if (FRunnableThread::Create( this, TEXT("StorageServerPrecache"), 0, EThreadPriority::TPri_Normal) == nullptr)
+		{
+			IsCompleted->Trigger();
+		}
+	}
+
+	virtual ~FStorageServerEngineStartupPrecache()
+	{
+		// the worker thread has most likely completed already but no harm in shutting down cleanly
+		bExitImmediately = true;
+		HasWork->Trigger();
+		IsCompleted->Wait();
+
+		FPlatformProcess::ReturnSynchEventToPool(HasWork);
+		FPlatformProcess::ReturnSynchEventToPool(IsCompleted);
+
+#if STORAGE_SERVER_START_CACHE_REPORT
+		UE_CLOG(Cache.Num() > 0, LogStorageServerPlatformFile, Log, TEXT("Engine startup precache size was %llu bytes, %d items. %d items actually retrieved (%d%%). %d cache misses (%d%%)"), CacheSize.load(), Cache.Num(), AccessedChunks.Num(), (AccessedChunks.Num()*100) / Cache.Num(), UncachedChunks.Num(), (UncachedChunks.Num()*100) / (AccessedChunks.Num()+UncachedChunks.Num()) );
+#endif
+	}
+
+	void AddPrecachedFile( const FIoChunkId& Id, uint32 Size, bool bHighPriority )
+	{
+		if (bHighPriority)
+		{
+			HighPriorityPrecacheFileRequests.Enqueue( TTuple<FIoChunkId,int32>(Id, Size) );
+		}
+		else
+		{
+			PrecacheFileRequests.Enqueue( TTuple<FIoChunkId,int32>(Id, Size) );
+		}
+		HasWork->Trigger();
+	}
+
+	void Finalize()
+	{
+		// no more precache requests after this - thread will shut down when precaching has finished
+		bExitWhenComplete = true;
+		HasWork->Trigger();
+	}
+
+	bool GetPrecachedFile( const FIoChunkId& Id, int64 Offset, int64 Size, uint8* Destination, int64& BytesRead)
+	{
+		bool bSuccess = false;
+		{
+			FReadScopeLock Lock(CacheLock);
+			if (const TArray<uint8>* BufferPtr = Cache.Find(Id))
+			{
+				if (ensure(Offset < BufferPtr->Num()))
+				{
+					BytesRead = FMath::Min(BufferPtr->Num() - Offset, Size);
+					FMemory::Memcpy(Destination, BufferPtr->GetData() + Offset, BytesRead);
+					bSuccess = true;
+				}
+			}
+		}
+
+#if STORAGE_SERVER_START_CACHE_REPORT
+		if (bSuccess)
+		{
+			FWriteScopeLock Lock(CacheLock);
+			AccessedChunks.Add(Id);
+		}
+		else
+		{
+			FWriteScopeLock Lock(CacheLock);
+			UncachedChunks.Add(Id);
+		}
+#endif
+
+		return bSuccess;
+	}
+
+private:
+	virtual uint32 Run() override
+	{
+		TRACE_CPUPROFILER_EVENT_SCOPE(StorageServerPlatformEngineStartupPrecache);
+
+		while (!bExitImmediately)
+		{
+			HasWork->Wait();
+			if (bExitImmediately)
+			{
+				break;
+			}
+
+			PrecacheItems();
+			if (bExitWhenComplete)
+			{
+				break;
+			}
+		}
+
+		IsCompleted->Trigger();
+		return 0;
+	}
+
+	void PrecacheItems()
+	{
+		TTuple<FIoChunkId,int32> Item;
+		while (HighPriorityPrecacheFileRequests.Dequeue(Item) || PrecacheFileRequests.Dequeue(Item))
+		{
+			const FIoChunkId& FileChunkId = Item.Key;
+			const int32& FileSize = Item.Value;
+
+			// make sure the total cache size doesn't grow out of control
+			if (CacheSize+FileSize > MaxCacheSize)
+			{
+				continue;
+			}
+
+			TArray<uint8> TempBuffer;
+			TempBuffer.AddUninitialized(FileSize);
+
+			TRACE_CPUPROFILER_EVENT_SCOPE(FStorageServerEngineStartupPrecache::SendReadMessage);
+			TIoStatusOr<FIoBuffer> Result = Connection.ReadChunkRequest(FileChunkId, 0, FileSize, FIoBuffer(FIoBuffer::Wrap, TempBuffer.GetData(), FileSize), false);
+			if (Result.IsOk())
+			{
+				TempBuffer.SetNum(Result.ValueOrDie().GetSize());
+				CacheSize += TempBuffer.Num();
+				{
+					FWriteScopeLock Lock(CacheLock);
+					Cache.Add(FileChunkId, MoveTemp(TempBuffer));
+				}
+			}
+		}
+	}
+
+	FStorageServerConnection& Connection;
+
+	FRWLock CacheLock;
+	TMap<FIoChunkId,TArray<uint8>> Cache;
+	std::atomic<uint64> CacheSize;
+
+	TSpscQueue<TTuple<FIoChunkId,int32>> PrecacheFileRequests;
+	TSpscQueue<TTuple<FIoChunkId,int32>> HighPriorityPrecacheFileRequests;
+	FEvent* HasWork;
+	FEvent* IsCompleted;
+
+	std::atomic<bool> bExitWhenComplete;
+	std::atomic<bool> bExitImmediately;
+	
+#if STORAGE_SERVER_START_CACHE_REPORT
+	TSet<FIoChunkId> UncachedChunks;
+	TSet<FIoChunkId> AccessedChunks;
+#endif
+};
+
+
+TUniquePtr<FStorageServerEngineStartupPrecache> GStorageServerEngineStartupPrecache;
+
+#endif // WITH_STORAGE_SERVER_STARTUP_FILE_CACHE
+
+
+
+
+
+
+
 #if COUNTERSTRACE_ENABLED
 	TRACE_DECLARE_ATOMIC_FLOAT_COUNTER(StorageServerCache_HitRatioBytes, TEXT("ZenClient/FileCache/HitRatio"));
 	namespace
@@ -320,6 +516,23 @@ public:
 			return;
 		}
 
+#if WITH_STORAGE_SERVER_STARTUP_FILE_CACHE
+		// check engine startup cache
+		if ( FStorageServerEngineStartupPrecache* Precache = GStorageServerEngineStartupPrecache.Get() )
+		{
+			int64 BytesRead;
+			if (Precache->GetPrecachedFile( FileChunkId, FilePos, BytesToRead, Destination, BytesRead ))
+			{
+				STORAGESERVER_CACHEHIT(BytesRead);
+				BytesToRead -= BytesRead;
+				FilePos += BytesRead;
+				Destination += BytesRead;
+				return;
+			}
+		}
+#endif
+
+
 		// try to read existing data from cache
 		{
 			UE::TScopeLock Lock(*this);
@@ -750,6 +963,11 @@ bool FStorageServerPlatformFile::Initialize(IPlatformFile* Inner, const TCHAR* C
 #if !WITH_EDITOR
 		// Extensions for file types that will be assumed to be immutable - their time stamp will remain unchanged.
 		AssumedImmutableTimeStampExtensions.Add(TEXT("uplugin"));
+
+		// Extensions for file types that will be precached on startup to improve engine initialization time
+		EngineStartupPrecacheExtensions.Add(TEXT("uplugin"));
+		EngineStartupPrecacheExtensions.Add(TEXT("uproject"));
+		EngineStartupPrecacheExtensions.Add(TEXT("ini"));
 #endif
 
 		// Don't initialize the connection yet because we want to incorporate project file path information into the initialization.
@@ -843,6 +1061,14 @@ void FStorageServerPlatformFile::InitializeConnection()
 	Connection.Reset(new FStorageServerConnection());
 	if (Connection->Initialize(HostAddrs, HostPort, MakeBaseURI()))
 	{
+#if WITH_STORAGE_SERVER_STARTUP_FILE_CACHE
+		GStorageServerEngineStartupPrecache.Reset( new FStorageServerEngineStartupPrecache(*Connection.Get()) );
+		FCoreDelegates::OnFEngineLoopInitComplete.AddLambda( []()
+		{
+			GStorageServerEngineStartupPrecache.Reset();
+		});	
+#endif
+
 		if (SendGetFileListMessage())
 		{
 			if (bAllowPackageIo)
@@ -1342,6 +1568,12 @@ bool FStorageServerPlatformFile::IsAssumedImmutableTimeStampFilename(FStringView
 	return AssumedImmutableTimeStampExtensions.Contains(Ext);
 }
 
+bool FStorageServerPlatformFile::IsEngineStartupPrecachableFilename(FStringView InFilename) const
+{
+	FName Ext = FName(FPathViews::GetExtension(InFilename));
+	return EngineStartupPrecacheExtensions.Contains(Ext);
+}
+
 bool FStorageServerPlatformFile::MakeStorageServerPath(const TCHAR* LocalFilenameOrDirectory, FStringBuilderBase& OutPath) const
 {
 	FStringView LocalEngineDirView(FPlatformMisc::EngineDir());
@@ -1408,8 +1640,20 @@ bool FStorageServerPlatformFile::SendGetFileListMessage()
 		}
 
 		ServerToc.AddFile(Id, Path, RawSize);
+
+#if WITH_STORAGE_SERVER_STARTUP_FILE_CACHE
+		if (RawSize > 0 && RawSize < FStorageServerEngineStartupPrecache::MaxFileSize && IsEngineStartupPrecachableFilename(Path) )
+		{
+			const bool bHighPriority = Path.EndsWith(TEXT(".uproject")) || Path.EndsWith(TEXT("DataDrivenPlatformInfo.ini")); // special case: we know the .uproject and DDPI will be needed immediately after this
+			GStorageServerEngineStartupPrecache->AddPrecachedFile(Id, RawSize, bHighPriority);
+		}
+#endif
 	});
 
+#if WITH_STORAGE_SERVER_STARTUP_FILE_CACHE
+	GStorageServerEngineStartupPrecache->Finalize();
+#endif
+
 	return true;
 }
 
diff --git a/Engine/Source/Runtime/StorageServerClient/Private/StorageServerPlatformFile.h b/Engine/Source/Runtime/StorageServerClient/Private/StorageServerPlatformFile.h
--- a/Engine/Source/Runtime/StorageServerClient/Private/StorageServerPlatformFile.h
+++ b/Engine/Source/Runtime/StorageServerClient/Private/StorageServerPlatformFile.h
@@ -131,6 +131,7 @@ private:
 	void InitializeConnection();
 	bool IsNonServerFilenameAllowed(FStringView InFilename);
 	bool IsAssumedImmutableTimeStampFilename(FStringView InFilename) const;
+	bool IsEngineStartupPrecachableFilename(FStringView InFilename) const;
 	bool MakeStorageServerPath(const TCHAR* LocalFilenameOrDirectory, FStringBuilderBase& OutPath) const;
 	bool MakeLocalPath(const TCHAR* ServerFilenameOrDirectory, FStringBuilderBase& OutPath) const;
 	IFileHandle* InternalOpenFile(const FIoChunkId& FileChunkId, int64 RawSize, const TCHAR* LocalFilename);
@@ -146,6 +147,7 @@ private:
 	FString CustomProjectStorePath;
 	TSet<FName> ExcludedNonServerExtensions;
 	TSet<FName> AssumedImmutableTimeStampExtensions;
+	TSet<FName> EngineStartupPrecacheExtensions;
 	IPlatformFile* LowerLevel = nullptr;
 	FStringView ServerEngineDirView = FStringView(TEXT("/{engine}/"));
 	FStringView ServerProjectDirView = FStringView(TEXT("/{project}/"));
