Zen client side caching basic infrastructure and generic implementation

diff --git a/Engine/Source/Runtime/Core/Public/ProfilingDebugging/CountersTrace.h b/Engine/Source/Runtime/Core/Public/ProfilingDebugging/CountersTrace.h
--- a/Engine/Source/Runtime/Core/Public/ProfilingDebugging/CountersTrace.h
+++ b/Engine/Source/Runtime/Core/Public/ProfilingDebugging/CountersTrace.h
@@ -387,6 +387,9 @@ struct FCountersTrace
 #define TRACE_COUNTER_SET(CounterName, Value) \
 	PREPROCESSOR_JOIN(__GTraceCounter, CounterName).Set(Value);
 
+#define TRACE_COUNTER_GET(CounterName) \
+	PREPROCESSOR_JOIN(__GTraceCounter, CounterName).Get();
+
 #define TRACE_COUNTER_ADD(CounterName, Value) \
 	PREPROCESSOR_JOIN(__GTraceCounter, CounterName).Add(Value);
 
diff --git a/Engine/Source/Runtime/StorageServerClient/Private/Cache/CacheJournal.h b/Engine/Source/Runtime/StorageServerClient/Private/Cache/CacheJournal.h
new file mode 100644
--- /dev/null
+++ b/Engine/Source/Runtime/StorageServerClient/Private/Cache/CacheJournal.h
@@ -0,0 +1,153 @@
+﻿// Copyright Epic Games, Inc. All Rights Reserved.
+
+#pragma once
+
+#include "IO/IoChunkId.h"
+#include "IO/IoHash.h"
+#include "Templates/Function.h"
+#include "Math/Interval.h"
+#include "StorageServerHttpClient.h"
+
+#if !UE_BUILD_SHIPPING
+
+namespace StorageServer
+{
+	// Opaque identifier that changes when content of chunk changes,
+	// right now is implemented via ChunkId RawHash as a stop gap solution
+	typedef FIoHash FModHash;
+
+	struct FCacheChunkInfo
+	{
+		FModHash ModHash;   // content modification chunk hash
+		int64 RawSize;      // size of complete chunk
+		int32 RawBlockSize; // size of block in a chunk if any, e.g. 256kb, can be 0
+
+		bool SetChunkInfo(
+			const FModHash* OptModHash,
+			const int64* OptRawSize,
+			const int32* OptRawBlockSize
+		)
+		{
+			bool bResult = true;
+			if (OptModHash)
+			{
+				bResult = ModHash == *OptModHash;
+				ModHash = *OptModHash;
+			}
+			if (OptRawSize)
+			{
+				RawSize = *OptRawSize;
+			}
+			if (OptRawBlockSize)
+			{
+				RawBlockSize = *OptRawBlockSize;
+			}
+			return bResult;
+		}
+
+		friend FArchive& operator << (FArchive& Ar, FCacheChunkInfo& Info)
+		{
+			Ar << Info.ModHash;
+			Ar << Info.RawSize;
+			Ar << Info.RawBlockSize;
+			return Ar;
+		}
+	};
+
+	struct FCacheEntry
+	{
+		uint64 ChunkOffset;   // uncompressed offset in a chunk
+		uint64 ChunkSize;     // uncompressed size in a chunk
+		uint64 StorageOffset; // offset in cache storage
+		uint64 StorageSize;   // size in cache storage
+		FIoHash StorageHash;  // hash of data in cache storage, needed to validate after reading from storage as we might lose data
+		EStorageServerContentType StorageContentType; // content type in storage
+
+		TInterval<uint64> GetChunkInterval() const
+		{
+			return TInterval(ChunkOffset, ChunkSize);
+		}
+
+		friend FArchive& operator << (FArchive& Ar, FCacheEntry& Entry)
+		{
+			Ar << Entry.ChunkOffset;
+			Ar << Entry.ChunkSize;
+			Ar << Entry.StorageOffset;
+			Ar << Entry.StorageSize;
+			Ar << Entry.StorageHash;
+			Ar << Entry.StorageContentType;
+			return Ar;
+		}
+	};
+
+	// Generic journal for storage server caching. Transactions are thread safe, atomic, saved at best effort.
+	// Cache journal focuses on avoiding data corruption when saving to disk, unlike cache storage which doesn't provide any guarantees for data consistency.
+	class ICacheJournal
+	{
+	public:
+		virtual ~ICacheJournal() = default;
+
+		// Flushes data to backing storage
+		virtual void Flush() = 0;
+
+		// Invalidates all data in journal
+		virtual void InvalidateAll() = 0;
+
+		// Invalidates data for a specific chunk.
+		virtual void Invalidate(const FIoChunkId& ChunkId) = 0;
+
+		// Updates chunk info. All fields are optional.
+		// If OptModHash is passed will compare it to existing ModHash and invalidate all data for this chunk id if they don't match.
+		// OptRawSize, OptRawBlockSize will simply be updated if passed.
+		// Returns true if either new ChunkId info is created or ModHash matches previous existing entry for this ChunkId.
+		// Returns false if ModHash doesn't match and cached data for this entry needs to be invalidated.
+		virtual bool SetChunkInfo(
+			const FIoChunkId& ChunkId,
+			const FModHash* OptModHash,
+			const int64* OptRawSize,
+			const int32* OptRawBlockSize
+		) = 0;
+
+		// Return chunk info if present, returns true if chunk info is present.
+		virtual bool TryGetChunkInfo(
+			const FIoChunkId& ChunkId,
+			FCacheChunkInfo& OutChunkInfo
+		) = 0;
+
+		// Adds a new cache entry for chunk.
+		// Returns true if new entry was added, false if there is an entry for (ChunkId / ChunkOffset / ChunkSize).
+		virtual bool AddEntry(
+			const FIoChunkId& ChunkId,
+			const FCacheEntry& Entry
+		) = 0;
+
+		// Tries to find an entry for specified chunk offset and size, returns true if entry is present.
+		virtual bool TryGetEntry(
+			const FIoChunkId& ChunkId,
+			const uint64 ChunkOffset,
+			const uint64 ChunkSize,
+			FCacheEntry& OutEntry
+		) = 0;
+
+		// Iterate chunk id in journal that contain at least one entry.
+		// Not safe to invoke other journal methods from the callback.
+		virtual void IterateChunkIds(
+			TFunctionRef<void(const FIoChunkId& ChunkId)> Callback
+		) = 0;
+
+		// Iterate entries for a chunk id in journal.
+		// Not safe to invoke other journal methods from the callback.
+		virtual void IterateCacheEntriesForChunkId(
+			const FIoChunkId& ChunkId,
+			TFunctionRef<void(const FCacheEntry& Entry)> Callback
+		) = 0;
+
+		// Iterate all entries in journal.
+		// Not safe to invoke other journal methods from the callback.
+		virtual void IterateCacheEntries(
+			TFunctionRef<void(const FIoChunkId& ChunkId, const FCacheEntry& Entry)> Callback
+		) = 0;
+	};
+}
+
+#endif
\ No newline at end of file
diff --git a/Engine/Source/Runtime/StorageServerClient/Private/Cache/CacheJournalSimple.cpp b/Engine/Source/Runtime/StorageServerClient/Private/Cache/CacheJournalSimple.cpp
new file mode 100644
--- /dev/null
+++ b/Engine/Source/Runtime/StorageServerClient/Private/Cache/CacheJournalSimple.cpp
@@ -0,0 +1,262 @@
+﻿// Copyright Epic Games, Inc. All Rights Reserved.
+
+#include "CacheJournalSimple.h"
+
+#if !UE_BUILD_SHIPPING
+
+#include "HAL/FileManager.h"
+#include "Misc/ScopeLock.h"
+#include "ProfilingDebugging/CountersTrace.h"
+
+TRACE_DECLARE_INT_COUNTER(ZenJournalSimpleChunks, TEXT("ZenClient/JournalSimple/Chunks"));
+TRACE_DECLARE_INT_COUNTER(ZenJournalSimpleEntries, TEXT("ZenClient/JournalSimple/Entries"));
+TRACE_DECLARE_INT_COUNTER(ZenJournalSimpleInvalidated, TEXT("ZenClient/JournalSimple/Invalidated"));
+
+namespace StorageServer
+{
+	FCacheJournalSimple::FCacheJournalSimple(const TCHAR* InFileName, const uint64 InFlushAtWriteCount)
+	{
+		FileName = InFileName;
+		FlushAtWriteCount = InFlushAtWriteCount;
+		CurrentWriteCount = 0;
+		bDirty = false;
+		FArchive* Reader = IFileManager::Get().CreateFileReader(*FileName);
+		if (Reader)
+		{
+			FArchive& Ar = *Reader; // also update writer part!
+
+			Ar << ChunkInfos;
+			Ar << ChunkCacheEntries;
+			delete Reader;
+		}
+
+		for (const TTuple<FIoChunkId, FPerChunkCacheEntries>& ChunkEntriesPair : ChunkCacheEntries)
+		{
+			TRACE_COUNTER_INCREMENT(ZenJournalSimpleChunks);
+			TRACE_COUNTER_ADD(ZenJournalSimpleEntries, ChunkEntriesPair.Value.Num());
+		}
+	}
+
+	FCacheJournalSimple::~FCacheJournalSimple()
+	{
+	}
+
+	void FCacheJournalSimple::Flush()
+	{
+		TRACE_CPUPROFILER_EVENT_SCOPE(FCacheJournalSimple::Flush);
+
+		if (!bDirty) // no need to sync this
+		{
+			return;
+		}
+
+		FScopeLock Lock(&FlushTaskLock);
+
+		if (!FlushTask.IsCompleted())
+		{
+			return;
+		}
+
+		FlushTask = UE::Tasks::Launch(TEXT("CacheJournalFlush"), [this]()
+		{
+			TRACE_CPUPROFILER_EVENT_SCOPE(FCacheJournalSimple::Flush::Task);
+
+			FScopeLock Lock(&DataLock);
+
+			if (!bDirty)
+			{
+				return;
+			}
+
+			// TODO avoid writing whole file every time
+			FArchive* Writer = IFileManager::Get().CreateFileWriter(*FileName);
+			if (!Writer)
+			{
+				return;
+			}
+
+			FArchive& Ar = *Writer;
+
+			// As we only care about chunkid's that are actually cached there is no need to store all chunk id's
+			decltype(ChunkInfos) OnlyCachedChunksInfo;
+
+			for (const TTuple<FIoChunkId, FPerChunkCacheEntries>& Pair : ChunkCacheEntries)
+			{
+				FCacheChunkInfo* Info = nullptr;
+				if ((Info = ChunkInfos.Find(Pair.Key)) != nullptr)
+				{
+					OnlyCachedChunksInfo.Add(Pair.Key, *Info);
+				}
+			}
+
+			Ar << OnlyCachedChunksInfo;
+			Ar << ChunkCacheEntries;
+
+			Writer->Flush();
+			delete Writer;
+			bDirty = false;
+		});
+	}
+
+	void FCacheJournalSimple::InvalidateAll()
+	{
+		FScopeLock Lock(&DataLock);
+		ChunkInfos.Empty();
+		ChunkCacheEntries.Empty();
+		TRACE_COUNTER_SET(ZenJournalSimpleChunks, 0);
+		TRACE_COUNTER_SET(ZenJournalSimpleEntries, 0);
+		bDirty = true;
+	}
+
+	void FCacheJournalSimple::Invalidate(const FIoChunkId& ChunkId)
+	{
+		FScopeLock Lock(&DataLock);
+
+		FPerChunkCacheEntries* CacheEntries = ChunkCacheEntries.Find(ChunkId);
+		if (CacheEntries)
+		{
+			TRACE_COUNTER_DECREMENT(ZenJournalSimpleChunks);
+			TRACE_COUNTER_SUBTRACT(ZenJournalSimpleEntries, CacheEntries->Num());
+		}
+
+		ChunkInfos.Remove(ChunkId);
+		ChunkCacheEntries.Remove(ChunkId);
+		bDirty = true;
+	}
+
+	bool FCacheJournalSimple::SetChunkInfo(const FIoChunkId& ChunkId, const FIoHash* OptModHash, const int64* OptRawSize, const int32* OptRawBlockSize)
+	{
+		FScopeLock Lock(&DataLock);
+		FCacheChunkInfo* ChunkInfo = ChunkInfos.Find(ChunkId);
+		if (ChunkInfo)
+		{
+			FModHash PreviousHash = ChunkInfo->ModHash;
+			const bool bHashesMatch = ChunkInfo->SetChunkInfo(OptModHash, OptRawSize, OptRawBlockSize);
+			if (!bHashesMatch)
+			{
+				TRACE_COUNTER_INCREMENT(ZenJournalSimpleInvalidated);
+			}
+			return bHashesMatch;
+		}
+
+		FCacheChunkInfo NewChunkInfo = {};
+		NewChunkInfo.SetChunkInfo(OptModHash, OptRawSize, OptRawBlockSize);
+		ChunkInfos.Add(ChunkId, NewChunkInfo);
+		bDirty = true;
+		return true;
+	}
+
+	bool FCacheJournalSimple::TryGetChunkInfo(const FIoChunkId& ChunkId, FCacheChunkInfo& OutChunkInfo)
+	{
+		FScopeLock Lock(&DataLock);
+		FCacheChunkInfo* ChunkInfo = ChunkInfos.Find(ChunkId);
+		if (!ChunkInfo)
+		{
+			return false;
+		}
+		OutChunkInfo = *ChunkInfo;
+		return true;
+	}
+
+	bool FCacheJournalSimple::AddEntry(const FIoChunkId& ChunkId, const FCacheEntry& Entry)
+	{
+		bool bNeedsFlush = false;
+
+		{
+			FScopeLock Lock(&DataLock);
+
+			FPerChunkCacheEntries* CacheEntries = ChunkCacheEntries.Find(ChunkId);
+			if (!CacheEntries)
+			{
+				CacheEntries = &ChunkCacheEntries.Add(ChunkId);
+				TRACE_COUNTER_INCREMENT(ZenJournalSimpleChunks);
+			}
+
+			const TInterval<uint64> ChunkInterval = Entry.GetChunkInterval();
+			if (CacheEntries->Contains(ChunkInterval))
+			{
+				return false;
+			}
+		
+			CacheEntries->Add(ChunkInterval, Entry);
+			TRACE_COUNTER_INCREMENT(ZenJournalSimpleEntries);
+			bDirty = true;
+
+			if (CurrentWriteCount++ > FlushAtWriteCount && FlushAtWriteCount > 0)
+			{
+				bNeedsFlush = true;
+				CurrentWriteCount = 0;
+			}
+		}
+
+		if (bNeedsFlush)
+		{
+			Flush();
+		}
+
+		return true;
+	}
+
+	bool FCacheJournalSimple::TryGetEntry(const FIoChunkId& ChunkId, const uint64 ChunkOffset, const uint64 ChunkSize, FCacheEntry& OutEntry)
+	{
+		FScopeLock Lock(&DataLock);
+
+		FPerChunkCacheEntries* CacheEntries = ChunkCacheEntries.Find(ChunkId);
+		if (!CacheEntries)
+		{
+			return false;
+		}
+
+		const TInterval<uint64> ChunkInterval(ChunkOffset, ChunkSize);
+		FCacheEntry* Entry = CacheEntries->Find(ChunkInterval);
+
+		if (!Entry)
+		{
+			return false;
+		}
+
+		OutEntry = *Entry;
+		return true;
+	}
+
+	void FCacheJournalSimple::IterateChunkIds(TFunctionRef<void(const FIoChunkId& ChunkId)> Callback)
+	{
+		FScopeLock Lock(&DataLock);
+
+		for (const TTuple<FIoChunkId, FPerChunkCacheEntries>& Pair : ChunkCacheEntries)
+		{
+			Callback(Pair.Key);
+		}
+	}
+
+	void FCacheJournalSimple::IterateCacheEntriesForChunkId(const FIoChunkId& ChunkId, TFunctionRef<void(const FCacheEntry& Entry)> Callback)
+	{
+		FScopeLock Lock(&DataLock);
+	
+		FPerChunkCacheEntries* CacheEntries = ChunkCacheEntries.Find(ChunkId);
+		if (!CacheEntries)
+		{
+			return;
+		}
+
+		for (const TTuple<TInterval<uint64>, FCacheEntry>& EntryPair : *CacheEntries)
+		{
+			Callback(EntryPair.Value);
+		}
+	}
+
+	void FCacheJournalSimple::IterateCacheEntries(TFunctionRef<void(const FIoChunkId& ChunkId, const FCacheEntry& Entry)> Callback)
+	{
+		FScopeLock Lock(&DataLock);
+
+		for (const TTuple<FIoChunkId, FPerChunkCacheEntries>& ChunkEntriesPair : ChunkCacheEntries)
+		{
+			for (const TTuple<TInterval<uint64>, FCacheEntry>& EntryPair : ChunkEntriesPair.Value)
+			{
+				Callback(ChunkEntriesPair.Key, EntryPair.Value);
+			}
+		}
+	}
+}
+
+#endif
diff --git a/Engine/Source/Runtime/StorageServerClient/Private/Cache/CacheJournalSimple.h b/Engine/Source/Runtime/StorageServerClient/Private/Cache/CacheJournalSimple.h
new file mode 100644
--- /dev/null
+++ b/Engine/Source/Runtime/StorageServerClient/Private/Cache/CacheJournalSimple.h
@@ -0,0 +1,74 @@
+﻿// Copyright Epic Games, Inc. All Rights Reserved.
+
+#pragma once
+
+#include "CacheJournal.h"
+#include "Containers/Map.h"
+#include "HAL/CriticalSection.h"
+#include "Tasks/Task.h"
+#include <atomic>
+
+#if !UE_BUILD_SHIPPING
+
+namespace StorageServer
+{
+	// Cross-platform cache journal implementation, baseline for platforms where platform specific implementation is not available.
+	class FCacheJournalSimple : public ICacheJournal
+	{
+	public:
+		FCacheJournalSimple(const TCHAR* FileName, const uint64 FlushAtWriteCount);
+		virtual ~FCacheJournalSimple() override;
+
+		virtual void Flush() override;
+
+		virtual void InvalidateAll() override;
+		virtual void Invalidate(const FIoChunkId& ChunkId) override;
+
+		virtual bool SetChunkInfo(
+			const FIoChunkId& ChunkId,
+			const FIoHash* OptModHash,
+			const int64* OptRawSize,
+			const int32* OptRawBlockSize
+		) override;
+		virtual bool TryGetChunkInfo(
+			const FIoChunkId& ChunkId,
+			FCacheChunkInfo& OutChunkInfo
+		) override;
+
+		virtual bool AddEntry(
+			const FIoChunkId& ChunkId,
+			const FCacheEntry& Entry
+		) override;
+		virtual bool TryGetEntry(
+			const FIoChunkId& ChunkId,
+			const uint64 ChunkOffset,
+			const uint64 ChunkSize,
+			FCacheEntry& OutEntry
+		) override;
+
+		virtual void IterateChunkIds(
+			TFunctionRef<void(const FIoChunkId& ChunkId)> Callback
+		) override;
+		virtual void IterateCacheEntriesForChunkId(
+			const FIoChunkId& ChunkId,
+			TFunctionRef<void(const FCacheEntry& Entry)> Callback
+		) override;
+		virtual void IterateCacheEntries(
+			TFunctionRef<void(const FIoChunkId& ChunkId, const FCacheEntry& Entry)> Callback
+		) override;
+		
+	private:
+		FCriticalSection DataLock, FlushTaskLock;
+		FString FileName;
+		UE::Tasks::TTask<void> FlushTask;
+		uint64 FlushAtWriteCount;
+		uint64 CurrentWriteCount;
+		std::atomic<bool> bDirty;
+
+		typedef TMap<TInterval<uint64>, FCacheEntry> FPerChunkCacheEntries;
+		TMap<FIoChunkId, FCacheChunkInfo> ChunkInfos;
+		TMap<FIoChunkId, FPerChunkCacheEntries> ChunkCacheEntries;
+	};
+}
+
+#endif
\ No newline at end of file
diff --git a/Engine/Source/Runtime/StorageServerClient/Private/Cache/CacheStorage.h b/Engine/Source/Runtime/StorageServerClient/Private/Cache/CacheStorage.h
new file mode 100644
--- /dev/null
+++ b/Engine/Source/Runtime/StorageServerClient/Private/Cache/CacheStorage.h
@@ -0,0 +1,50 @@
+﻿// Copyright Epic Games, Inc. All Rights Reserved.
+
+#pragma once
+
+#include "IO/IoChunkId.h"
+#include "IO/IoBuffer.h"
+
+#if !UE_BUILD_SHIPPING
+
+namespace StorageServer
+{
+	// Bulk storage for caching, transactions are async, non-atomic, best effort, data might be corrupted.
+	class ICacheStorage
+	{
+	public:
+		virtual ~ICacheStorage() = default;
+
+		// returns true if cache was created/truncated this session and all data should be assumed lost
+		virtual bool IsNewlyCreatedStorage()
+		{
+			return bNewlyCreatedStorage;
+		}
+
+		virtual void Flush() = 0;
+
+		virtual uint64 GetSize() const = 0;
+
+		virtual void Invalidate(const uint64 Offset, const uint64 Size) = 0;
+
+		virtual FIoBuffer Read(
+			const uint64 Offset,
+			const uint64 ReadSize,
+			TOptional<FIoBuffer> OptDestination
+		) = 0;
+
+		virtual void WriteAsync(
+			const uint64 Offset,
+			const void* Buffer,
+			const uint64 WriteSize
+		) = 0;
+
+	protected:
+		bool bNewlyCreatedStorage = false;
+
+		// Limit cache file size to 2GB to stay below file system limitations
+		static constexpr uint64 MaxCacheFileSize = 2ull * 1024 * 1024 * 1024;
+	};
+}
+
+#endif
\ No newline at end of file
diff --git a/Engine/Source/Runtime/StorageServerClient/Private/Cache/CacheStorageBlocking.cpp b/Engine/Source/Runtime/StorageServerClient/Private/Cache/CacheStorageBlocking.cpp
new file mode 100644
--- /dev/null
+++ b/Engine/Source/Runtime/StorageServerClient/Private/Cache/CacheStorageBlocking.cpp
@@ -0,0 +1,253 @@
+﻿// Copyright Epic Games, Inc. All Rights Reserved.
+
+#include "CacheStorageBlocking.h"
+
+#if !UE_BUILD_SHIPPING
+
+#include "HAL/PlatformFileManager.h"
+#include "GenericPlatform/GenericPlatformFile.h"
+#include "ProfilingDebugging/CountersTrace.h"
+
+TRACE_DECLARE_MEMORY_COUNTER(FrontBuffer, TEXT("ZenClient/BlockingStorage/FrontBuffer"));
+TRACE_DECLARE_MEMORY_COUNTER(BackBuffer, TEXT("ZenClient/BlockingStorage/BackBuffer"));
+
+namespace StorageServer
+{
+	FCacheStorageBlocking::FCacheStorageBlocking(const TCHAR* FileNamePrefix, const uint64 FileSizeTmp)
+	{
+		IPlatformFile& PlatformFile = FPlatformFileManager::Get().GetPlatformFile();
+
+		BackingFiles.SetNum(FMath::DivideAndRoundUp(FileSizeTmp, MaxCacheFileSize));
+		FileSize = 0;
+
+		const uint64 SizeOfLastFile = FileSizeTmp - MaxCacheFileSize * (BackingFiles.Num() - 1);
+
+		for (int32 i = 0; i < BackingFiles.Num(); ++i)
+		{
+			// TODO check if disk size is exhausted?
+			// TODO use temp location that is available on development platforms
+			FString FileName = FString::Printf(TEXT("%s%u"), FileNamePrefix, i);
+			IFileHandle* BackingFile = PlatformFile.OpenWrite(*FileName, true, true);
+			if (!ensureAlwaysMsgf(BackingFile, TEXT("Can't open storage server backing file '%s'"), *FileName))
+			{
+				return;
+			}
+
+			const uint64 DesiredFileSize = i == BackingFiles.Num() - 1 ? SizeOfLastFile : MaxCacheFileSize;
+
+			if (BackingFile->Size() != DesiredFileSize)
+			{
+				BackingFile->Truncate(DesiredFileSize);
+				bNewlyCreatedStorage = true;
+			}
+
+			FileSize += BackingFile->Size();
+
+			BackingFiles[i] = TUniquePtr<IFileHandle>(BackingFile);
+		}
+	}
+
+	FCacheStorageBlocking::~FCacheStorageBlocking()
+	{
+		FlushTask.Wait();
+	}
+
+	void FCacheStorageBlocking::Flush()
+	{
+		// don't block if there is a pending flush already
+		Flush(true);
+	}
+
+	uint64 FCacheStorageBlocking::GetSize() const
+	{
+		return FileSize;
+	}
+
+	FIoBuffer FCacheStorageBlocking::Read(const uint64 Offset, const uint64 ReadSize, TOptional<FIoBuffer> OptDestination)
+	{
+		// TODO potential improvement is to inspect in-flight write queues and read data from them if needed
+		FIoBuffer Destination = OptDestination.IsSet() && OptDestination->GetSize() >= ReadSize ? OptDestination.GetValue() : FIoBuffer(ReadSize);
+		Destination.SetSize(ReadSize);
+		BackingReadAt(Offset, Destination.GetData(), ReadSize);
+		return Destination;
+	}
+
+	void FCacheStorageBlocking::WriteAsync(const uint64 Offset, const void* Buffer, const uint64 WriteSize)
+	{
+		TRACE_CPUPROFILER_EVENT_SCOPE(FCacheStorageBlocking::WriteAsync);
+
+		// if buffer is too big it's easier to just store it directly
+		if (WriteSize >= Buffers.GetBack()->MaxDataSize / 2)
+		{
+			TRACE_CPUPROFILER_EVENT_SCOPE(FCacheStorageBlocking::WriteAsync::SyncWrite);
+			BackingWriteAt(Offset, Buffer, WriteSize);
+		}
+		else
+		{
+			TRACE_CPUPROFILER_EVENT_SCOPE(FCacheStorageBlocking::WriteAsync::Write);
+			bool bNeedsFlush = false;
+			while (true)
+			{
+				if (bNeedsFlush)
+				{
+					// wait for buffer swap
+					Flush(false);
+				}
+
+				FScopeLock ScopeLock(&BackLock);
+				if (Buffers.GetBack()->Write(Offset, Buffer, WriteSize))
+				{
+					return;
+				}
+				else
+				{
+					bNeedsFlush = true;
+				}
+			}
+		}
+	}
+
+	void FCacheStorageBlocking::Flush(bool bLazy)
+	{
+		TRACE_CPUPROFILER_EVENT_SCOPE(FCacheStorageBlocking::Flush);
+
+		FScopeLock ScopeBackLock(&BackLock);
+
+		if (!FlushTask.IsCompleted())
+		{
+			if (bLazy)
+			{
+				return;
+			}
+			else
+			{
+				TRACE_CPUPROFILER_EVENT_SCOPE(FCacheStorageBlocking::Flush::FlushTaskWait);
+				FlushTask.Wait();
+			}
+		}
+
+		// no data in back buffer, nothing to flush
+		if (Buffers.GetBack()->Data.Num() == 0)
+		{
+			return;
+		}
+
+		Buffers.Flip();
+
+		TRACE_COUNTER_SET(FrontBuffer, Buffers.GetFront()->Data.Num());
+		TRACE_COUNTER_SET(BackBuffer, Buffers.GetBack()->Data.Num());
+
+		FlushTask = UE::Tasks::Launch(TEXT("CacheStorageFlush"), [this]()
+		{
+			FScopeLock ScopeFrontLock(&FrontLock);
+			FWriteQueue* Front = Buffers.GetFront();
+
+			for (const FWriteOp& Operation : Front->Operations)
+			{
+				TRACE_CPUPROFILER_EVENT_SCOPE(FCacheStorageBlocking::Flush::Write);
+				BackingWriteAt(Operation.OffsetInFile, Operation.Data, Operation.Size);
+				TRACE_COUNTER_SUBTRACT(FrontBuffer, Operation.Size);
+			}
+
+			Front->Empty();
+		}, LowLevelTasks::ETaskPriority::BackgroundNormal);
+	}
+
+	void FCacheStorageBlocking::BackingIntervals(const uint64 Offset, const uint64 Size,
+												 uint32& IndexA, uint64& OffsetA, uint64& SizeA,
+												 uint32& IndexB, uint64& OffsetB, uint64& SizeB) const
+	{
+		IndexA = FMath::DivideAndRoundDown(Offset, MaxCacheFileSize);
+		IndexB = FMath::DivideAndRoundDown(Offset + Size - 1, MaxCacheFileSize);
+
+		// we only support writes spanning maximum two backing storage files
+		ensureAlways(IndexA + 1 >= IndexB);
+
+		if (IndexA == IndexB)
+		{
+			OffsetA = Offset - IndexA * MaxCacheFileSize;
+			SizeA = Size;
+			OffsetB = 0;
+			SizeB = 0;
+		}
+		else
+		{
+			OffsetA = Offset - IndexA * MaxCacheFileSize;
+			SizeA = IndexB * MaxCacheFileSize - Offset;
+			OffsetB = 0;
+			SizeB = Size - SizeA;
+		}
+	}
+
+	void FCacheStorageBlocking::BackingReadAt(const uint64 Offset, void* Buffer, const uint64 Size)
+	{
+		uint32 IndexA, IndexB;
+		uint64 OffsetA, SizeA, OffsetB, SizeB;
+		BackingIntervals(Offset, Size, IndexA, OffsetA, SizeA, IndexB, OffsetB, SizeB);
+
+		if (SizeA > 0)
+		{
+			BackingFiles[IndexA]->ReadAt(static_cast<uint8*>(Buffer), SizeA, OffsetA);
+		}
+
+		if (SizeB > 0)
+		{
+			BackingFiles[IndexB]->ReadAt(static_cast<uint8*>(Buffer) + SizeA, SizeB, OffsetB);
+		}
+	}
+
+	void FCacheStorageBlocking::BackingWriteAt(const uint64 Offset, const void* Buffer, const uint64 Size)
+	{
+		uint32 IndexA, IndexB;
+		uint64 OffsetA, SizeA, OffsetB, SizeB;
+		BackingIntervals(Offset, Size, IndexA, OffsetA, SizeA, IndexB, OffsetB, SizeB);
+
+		FScopeLock ScopeFileWriteLock(&FileWriteLock);
+
+		if (SizeA > 0)
+		{
+			BackingFiles[IndexA]->Seek(OffsetA);
+			BackingFiles[IndexA]->Write(static_cast<const uint8*>(Buffer), SizeA);
+		}
+
+		if (SizeB > 0)
+		{
+			BackingFiles[IndexB]->Seek(OffsetB);
+			BackingFiles[IndexB]->Write(static_cast<const uint8*>(Buffer) + SizeA, SizeB);
+		}
+	}
+
+	bool FCacheStorageBlocking::FWriteQueue::Write(const uint64 OffsetInFile, const void* Buffer, const uint64 Size)
+	{
+		if (!CanStore(Size))
+		{
+			return false;
+		}
+
+		TRACE_COUNTER_ADD(BackBuffer, Size);
+
+		uint8* Ptr = Data.GetData() + Data.Num();
+		Data.SetNum(Data.Num() + Size);
+		FMemory::Memcpy(Ptr, Buffer, Size);
+
+		FWriteOp* Last = Operations.Num() > 0 ? &Operations.Last() : nullptr;
+		// Try to coalesce write operation with last operation
+		// in hopes that in all linear writes we will only have one write op.
+		// Another way would be to sort them later and coalesce on sorted data.
+		if (Last != nullptr && Last->OffsetInFile + Last->Size == OffsetInFile)
+		{
+			Last->Size += Size;
+		}
+		else
+		{
+			FWriteOp& New = Operations.AddDefaulted_GetRef();
+			New.OffsetInFile = OffsetInFile;
+			New.Size = Size;
+			New.Data = Ptr;
+		}
+
+		return true;
+	}
+}
+
+#endif
diff --git a/Engine/Source/Runtime/StorageServerClient/Private/Cache/CacheStorageBlocking.h b/Engine/Source/Runtime/StorageServerClient/Private/Cache/CacheStorageBlocking.h
new file mode 100644
--- /dev/null
+++ b/Engine/Source/Runtime/StorageServerClient/Private/Cache/CacheStorageBlocking.h
@@ -0,0 +1,112 @@
+﻿// Copyright Epic Games, Inc. All Rights Reserved.
+
+#pragma once
+
+#include "CacheStorage.h"
+#include "Tasks/Task.h"
+
+#if !UE_BUILD_SHIPPING
+
+namespace StorageServer
+{
+	// Cache storage implemented via blocking file operations
+	class FCacheStorageBlocking : public ICacheStorage
+	{
+	public:
+		FCacheStorageBlocking(const TCHAR* FileName, const uint64 FileSize);
+		virtual ~FCacheStorageBlocking() override;
+
+		virtual void Flush() override;
+
+		virtual uint64 GetSize() const override;
+
+		virtual void Invalidate(const uint64 Offset, const uint64 Size) override {} // no-op
+
+		virtual FIoBuffer Read(
+			const uint64 Offset,
+			const uint64 ReadSize,
+			TOptional<FIoBuffer> OptDestination
+		) override;
+
+		virtual void WriteAsync(
+			const uint64 Offset,
+			const void* Buffer,
+			const uint64 WriteSize
+		) override;
+
+	private:
+		struct FWriteOp
+		{
+			uint64 OffsetInFile;
+			uint64 Size;
+			uint8* Data;
+		};
+
+		struct FWriteQueue
+		{
+			static constexpr uint32 MaxOpCount  =  1 * 1024;
+			static constexpr uint32 MaxDataSize = 16 * 1024 * 1024;
+
+			TArray<FWriteOp, TInlineAllocator<MaxOpCount * sizeof(FWriteOp)>> Operations;
+			TArray<uint8, TInlineAllocator<MaxDataSize>> Data;
+
+			bool CanStore(const uint64 Size) const
+			{
+				return Data.Num() + Size <= MaxDataSize;
+			}
+			bool Write(const uint64 OffsetInFile, const void* Buffer, const uint64 Size);
+			void Empty()
+			{
+				Operations.Empty();
+				Data.Empty();
+			}
+		};
+
+		class FDoubleBuffer
+		{
+		public:
+			FDoubleBuffer()
+				: Front(&Queue1)
+				, Back(&Queue2)
+			{}
+
+			FWriteQueue* GetFront() const {return Front;}
+			FWriteQueue* GetBack() const {return Back;}
+
+			void Flip()
+			{
+				if (Front == &Queue1)
+				{
+					Front = &Queue2;
+					Back = &Queue1;
+				}
+				else
+				{
+					Front = &Queue1;
+					Back = &Queue2;
+				}
+			}
+
+		private:
+			FWriteQueue Queue1;
+			FWriteQueue Queue2;
+			FWriteQueue* Front;
+			FWriteQueue* Back;
+		};
+
+		UE::Tasks::TTask<void> FlushTask;
+		TArray<TUniquePtr<IFileHandle>> BackingFiles;
+		FDoubleBuffer Buffers;
+		FCriticalSection FrontLock, BackLock, FileWriteLock;
+		uint64 FileSize;
+
+		void Flush(bool bLazy);
+		void BackingIntervals(const uint64 Offset, const uint64 Size,
+			uint32& IndexA, uint64& OffsetA, uint64& SizeA,
+			uint32& IndexB, uint64& OffsetB, uint64& SizeB) const;
+		void BackingReadAt(const uint64 Offset, void* Buffer, const uint64 Size);
+		void BackingWriteAt(const uint64 Offset, const void* Buffer, const uint64 Size);
+	};
+}
+
+#endif
\ No newline at end of file
diff --git a/Engine/Source/Runtime/StorageServerClient/Private/Cache/CacheStrategy.h b/Engine/Source/Runtime/StorageServerClient/Private/Cache/CacheStrategy.h
new file mode 100644
--- /dev/null
+++ b/Engine/Source/Runtime/StorageServerClient/Private/Cache/CacheStrategy.h
@@ -0,0 +1,119 @@
+﻿// Copyright Epic Games, Inc. All Rights Reserved.
+
+#pragma once
+
+#include "CacheJournal.h"
+#include "CacheStorage.h"
+#include "Templates/UniquePtr.h"
+#include "Containers/Ticker.h"
+
+#if !UE_BUILD_SHIPPING
+
+namespace StorageServer
+{
+	class ICacheStrategy
+	{
+	public:
+		ICacheStrategy(TUniquePtr<ICacheJournal>&& InJournal, TUniquePtr<ICacheStorage>&& InStorage, const float FlushInterval)
+			: Journal(MoveTemp(InJournal))
+			, Storage(MoveTemp(InStorage))
+		{
+			ensure(Journal.IsValid());
+			ensure(Storage.IsValid());
+
+			if (FlushInterval > 0.0f)
+			{
+				FlushTicker = FTSTicker::GetCoreTicker().AddTicker(FTickerDelegate::CreateRaw(this, &ICacheStrategy::FlushTick), FlushInterval);
+			}
+			else
+			{
+				FlushTicker.Reset();
+			}
+		}
+
+		virtual ~ICacheStrategy()
+		{
+			if (FlushTicker.IsValid())
+			{
+				FTSTicker::GetCoreTicker().RemoveTicker(FlushTicker);
+				FlushTicker.Reset();
+			}
+		}
+
+		virtual void Invalidate(const FIoChunkId& ChunkId) = 0;
+
+		// Stop gap implementation till new endpoint is available in zen
+		virtual void SetChunkInfo(
+			const FIoChunkId& ChunkId,
+			const FModHash& ModHash,
+			const int64& RawSize
+		)
+		{
+			if (!Journal->SetChunkInfo(ChunkId, &ModHash, &RawSize, nullptr))
+			{
+				Invalidate(ChunkId);
+			}
+		}
+
+		virtual void CacheChunkSize(const FIoChunkId& ChunkId, const int64 RawSize)
+		{
+			Journal->SetChunkInfo(ChunkId, nullptr, &RawSize, nullptr);
+		}
+
+		virtual bool TryGetChunkSize(const FIoChunkId& ChunkId, int64& OutRawSize)
+		{
+			FCacheChunkInfo ChunkInfo;
+			if (Journal->TryGetChunkInfo(ChunkId, ChunkInfo))
+			{
+				OutRawSize = ChunkInfo.RawSize;
+				return true;
+			}
+			else
+			{
+				return false;
+			}
+		}
+
+		virtual bool ReadChunk(
+			const FIoChunkId& RequestChunkId,
+			const uint64 RequestOffset,
+			const uint64 RequestSize,
+			TOptional<FIoBuffer> OptDestination,
+			FIoBuffer& OutBuffer,
+			EStorageServerContentType& OutContentType
+		) = 0;
+
+		// TODO cache individual chunk blocks.
+		// In case if ResultContentType is CompressedBinary the ResultBuffer will contain 1 or more chunk blocks,
+		// we can cache them individually with corresponding block offset and block size,
+		// such as if a follow-up ReadChunk request comes in that maps to the same block (but not same offset/size that was used in first request),
+		// it would still be fulfilled from cache.
+		// This requires unpacking of CompressedBinary and reconstructing its header for each individual chunk block.
+		virtual void CacheChunk(
+			const FIoChunkId& RequestChunkId,
+			const uint64 RequestRawOffset,
+			const uint64 RequestRawSize,
+			const FIoBuffer& ResultBuffer,
+			const EStorageServerContentType ResultContentType
+		) = 0;
+
+		virtual void Flush()
+		{
+			Journal->Flush();
+			Storage->Flush();
+		}
+
+	protected:
+		TUniquePtr<ICacheJournal> Journal;
+		TUniquePtr<ICacheStorage> Storage;
+		FTSTicker::FDelegateHandle FlushTicker;
+
+		bool FlushTick(float DeltaTime)
+		{
+			Flush();
+			return true;
+		}
+	};
+}
+
+#endif
\ No newline at end of file
diff --git a/Engine/Source/Runtime/StorageServerClient/Private/Cache/CacheStrategyLinear.cpp b/Engine/Source/Runtime/StorageServerClient/Private/Cache/CacheStrategyLinear.cpp
new file mode 100644
--- /dev/null
+++ b/Engine/Source/Runtime/StorageServerClient/Private/Cache/CacheStrategyLinear.cpp
@@ -0,0 +1,201 @@
+// Copyright Epic Games, Inc. All Rights Reserved.
+
+#include "CacheStrategyLinear.h"
+
+#define ZEN_CACHE_VERBOSE_LOG 0
+
+#if !UE_BUILD_SHIPPING
+
+#include "IO/IoBuffer.h"
+#include "Misc/App.h"
+#include "ProfilingDebugging/CountersTrace.h"
+
+TRACE_DECLARE_MEMORY_COUNTER(ZenLinearStrategyCacheCurrentSize, TEXT("ZenClient/LinearStrategyCache/CurrentSize"));
+TRACE_DECLARE_MEMORY_COUNTER(ZenLinearStrategyCacheInvalidSize, TEXT("ZenClient/LinearStrategyCache/InvalidSize"));
+
+#if ZEN_CACHE_VERBOSE_LOG
+DEFINE_LOG_CATEGORY_STATIC(LogCacheStrategyLinear, Log, All);
+#endif
+
+namespace StorageServer
+{
+	FCacheStrategyLinear::FCacheStrategyLinear(TUniquePtr<ICacheJournal>&& InJournal, TUniquePtr<ICacheStorage>&& InStorage, const uint64 InAbandonAtInvalidSize, const float FlushInterval, const bool bInvalidate)
+		: ICacheStrategy(MoveTemp(InJournal), MoveTemp(InStorage), FlushInterval)
+	{
+		uint64 NextFreeOffset = 0;
+		uint64 TotalStoredSize = 0;
+		Journal->IterateCacheEntries([&NextFreeOffset, &TotalStoredSize](const FIoChunkId&, const FCacheEntry& CacheEntry)
+		{
+			TotalStoredSize += CacheEntry.StorageSize;
+			NextFreeOffset = FMath::Max(NextFreeOffset, CacheEntry.StorageOffset + CacheEntry.StorageSize);
+		});
+
+		CurrentSize = NextFreeOffset;
+		CurrentInvalidSize = NextFreeOffset - TotalStoredSize;
+		AbandonAtInvalidSize = InAbandonAtInvalidSize;
+		SetCounters();
+
+#if ZEN_CACHE_VERBOSE_LOG
+		UE_LOG(LogCacheStrategyLinear, Display, TEXT("Zen linear cache %llu / %llu / %llu"), CurrentSize.load(), CurrentInvalidSize, AbandonAtInvalidSize);
+#endif
+
+		if (bInvalidate ||
+			Storage->IsNewlyCreatedStorage() ||
+			CurrentInvalidSize >= AbandonAtInvalidSize)
+		{
+			InvalidateAll();
+		}
+	}
+
+	FCacheStrategyLinear::~FCacheStrategyLinear()
+	{
+	}
+
+	void FCacheStrategyLinear::InvalidateAll()
+	{
+#if ZEN_CACHE_VERBOSE_LOG
+		UE_LOG(LogCacheStrategyLinear, Display, TEXT("Abandoning zen linear cache"));
+#endif
+
+		Journal->InvalidateAll();
+		CurrentSize = 0;
+		CurrentInvalidSize = 0;
+		SetCounters();
+	}
+
+	void FCacheStrategyLinear::Invalidate(const FIoChunkId& ChunkId)
+	{
+		if (CurrentSize == 0)
+		{
+			return;
+		}
+
+		Journal->IterateCacheEntriesForChunkId(ChunkId, [this](const FCacheEntry& CacheEntry)
+		{
+			CurrentInvalidSize += CacheEntry.StorageSize;
+		});
+
+		SetCounters();
+
+		if (CurrentInvalidSize >= AbandonAtInvalidSize)
+		{
+			InvalidateAll();
+		}
+	}
+
+	bool FCacheStrategyLinear::ReadChunk(const FIoChunkId& RequestChunkId, const uint64 RequestOffset, const uint64 RequestSize, TOptional<FIoBuffer> OptDestination, FIoBuffer& OutBuffer, EStorageServerContentType& OutContentType)
+	{
+		FCacheEntry Entry = {};
+		if (!Journal->TryGetEntry(RequestChunkId, RequestOffset, RequestSize, Entry))
+		{
+#if ZEN_CACHE_VERBOSE_LOG
+			UE_LOG(LogCacheStrategyLinear, Display, TEXT("CacheMiss Key %s %llu:%llu"),
+				*LexToString(RequestChunkId),
+				RequestOffset,
+				RequestSize);
+#endif
+			return false;
+		}
+
+		FIoBuffer Buffer = Storage->Read(Entry.StorageOffset, Entry.StorageSize, OptDestination);
+		if (Buffer.GetSize() != Entry.StorageSize)
+		{
+#if ZEN_CACHE_VERBOSE_LOG
+			UE_LOG(LogCacheStrategyLinear, Display, TEXT("CacheMiss Key %s %llu:%llu failed read"),
+				*LexToString(RequestChunkId),
+				RequestOffset,
+				RequestSize);
+#endif
+			return false;
+		}
+
+		// Compare hashes here because storage doesn't guarantee data consistency
+		const FIoHash StorageHash = FIoHash::HashBuffer(Buffer.GetData(), Buffer.GetSize());
+		if (Entry.StorageHash != StorageHash)
+		{
+#if ZEN_CACHE_VERBOSE_LOG
+			UE_LOG(LogCacheStrategyLinear, Display, TEXT("CacheMiss Hash %s %llu:%llu->%llu:%llu %s != %s"),
+				*LexToString(RequestChunkId),
+				RequestOffset,
+				RequestSize,
+				Entry.StorageOffset,
+				Entry.StorageSize,
+				*LexToString(Entry.StorageHash),
+				*LexToString(StorageHash)
+			);
+#endif
+			return false;
+		}
+
+#if ZEN_CACHE_VERBOSE_LOG
+		UE_LOG(LogCacheStrategyLinear, Display, TEXT("CacheHit %s %llu:%llu->%llu:%llu %s"),
+			*LexToString(RequestChunkId),
+			RequestOffset,
+			RequestSize,
+			Entry.StorageOffset,
+			Entry.StorageSize,
+			*LexToString(Entry.StorageHash)
+		);
+#endif
+
+		OutBuffer = Buffer;
+		OutContentType = Entry.StorageContentType;
+		return true;
+	}
+
+	void FCacheStrategyLinear::CacheChunk(const FIoChunkId& RequestChunkId, const uint64 RequestOffset, const uint64 RequestSize, const FIoBuffer& Buffer, const EStorageServerContentType ContentType)
+	{
+		// TODO add hint to ignore some requests to limit overall write speed
+
+		bool bCommitToStorage = false;
+		FCacheEntry Entry = {};
+
+		{
+			FScopeLock ScopeLock(&JournalLock);
+
+			if (CurrentSize.load() + Buffer.GetSize() >= Storage->GetSize())
+			{
+				// no storage left
+				return;
+			}
+
+			Entry.ChunkOffset = RequestOffset;
+			Entry.ChunkSize = RequestSize;
+			Entry.StorageOffset = CurrentSize.load();
+			Entry.StorageSize = Buffer.GetSize();
+			Entry.StorageHash = FIoHash::HashBuffer(Buffer.GetData(), Buffer.GetSize());
+			Entry.StorageContentType = ContentType;
+
+			if (Journal->AddEntry(RequestChunkId, Entry))
+			{
+				CurrentSize += Buffer.GetSize();
+				bCommitToStorage = true;
+				SetCounters();
+
+#if ZEN_CACHE_VERBOSE_LOG
+				UE_LOG(LogCacheStrategyLinear, Display, TEXT("CachePut %s %llu:%llu->%llu:%llu %s"),
+					*LexToString(RequestChunkId),
+					RequestOffset,
+					RequestSize,
+					Entry.StorageOffset,
+					Entry.StorageSize,
+					*LexToString(Entry.StorageHash)
+				);
+#endif
+			}
+		}
+
+		if (bCommitToStorage)
+		{
+			Storage->WriteAsync(Entry.StorageOffset, Buffer.GetData(), Buffer.GetSize());
+		}
+	}
+
+	void FCacheStrategyLinear::SetCounters()
+	{
+		TRACE_COUNTER_SET(ZenLinearStrategyCacheCurrentSize, CurrentSize.load());
+		TRACE_COUNTER_SET(ZenLinearStrategyCacheInvalidSize, CurrentInvalidSize);
+	}
+}
+
+#endif
diff --git a/Engine/Source/Runtime/StorageServerClient/Private/Cache/CacheStrategyLinear.h b/Engine/Source/Runtime/StorageServerClient/Private/Cache/CacheStrategyLinear.h
new file mode 100644
--- /dev/null
+++ b/Engine/Source/Runtime/StorageServerClient/Private/Cache/CacheStrategyLinear.h
@@ -0,0 +1,37 @@
+﻿// Copyright Epic Games, Inc. All Rights Reserved.
+
+#pragma once
+
+#include "CacheStrategy.h"
+#include <atomic>
+
+#if !UE_BUILD_SHIPPING
+
+namespace StorageServer
+{
+	// Use cache storage as arena allocator and store all chunks as they come in. No eviction.
+	class FCacheStrategyLinear : public ICacheStrategy
+	{
+	public:
+		FCacheStrategyLinear(TUniquePtr<ICacheJournal>&& InJournal, TUniquePtr<ICacheStorage>&& InStorage, const uint64 AbandonAtInvalidSize, const float FlushInterval, const bool bInvalidate);
+		virtual ~FCacheStrategyLinear() override;
+
+		void InvalidateAll();
+		virtual void Invalidate(const FIoChunkId& ChunkId) override;
+
+		virtual bool ReadChunk(const FIoChunkId& RequestChunkId, const uint64 RequestOffset, const uint64 RequestSize, TOptional<FIoBuffer> OptDestination, FIoBuffer& OutBuffer, EStorageServerContentType& ContentType) override;
+
+		virtual void CacheChunk(const FIoChunkId& RequestChunkId, const uint64 RequestOffset, const uint64 RequestSize, const FIoBuffer& Buffer, const EStorageServerContentType ContentType) override;
+
+	private:
+		// Linear cache is implemented as arena allocator
+		std::atomic<uint64> CurrentSize; // Current size of arena allocator
+		uint64 CurrentInvalidSize; // Amount of bytes with-in arena allocator that are stale
+		uint64 AbandonAtInvalidSize; // Consider cache to be a lost cause at this amount of stale data
+		FCriticalSection JournalLock;
+
+		void SetCounters();
+	};
+}
+
+#endif
\ No newline at end of file
diff --git a/Engine/Source/Runtime/StorageServerClient/Private/StorageServerConnection.cpp b/Engine/Source/Runtime/StorageServerClient/Private/StorageServerConnection.cpp
--- a/Engine/Source/Runtime/StorageServerClient/Private/StorageServerConnection.cpp
+++ b/Engine/Source/Runtime/StorageServerClient/Private/StorageServerConnection.cpp
@@ -22,13 +22,55 @@
 #include "BuiltInHttpClient/BuiltInHttpClientPlatformSocket.h"
 #include "HAL/PlatformMath.h"
 
+#include "Cache/CacheJournalSimple.h"
+#include "Cache/CacheStorageBlocking.h"
+#include "Cache/CacheStrategyLinear.h"
+
 #if !UE_BUILD_SHIPPING
 
+// Set to true to enable cache, temporarily disabled for now
+static constexpr bool bEnableZenCache = true; // PLATFORM_SUPPORTS_STORAGE_SERVER_CACHE;
+
+// TODO implement a way to control cache via cvars, right now it's not possible as cvars are initialized at later stage than cache
+extern "C" int32 GStorageServerLocalCacheSizeKB = bEnableZenCache ? 8 * 1024 * 1024 : 0; // 8 GB
+extern "C" float GStorageServerLocalCacheFlushInterval = 10.0;
+extern "C" int32 GStorageServerLocalCacheFlushEveryNEntries = 0; // set to >0 to flush journal every N new entries
+extern "C" int32 GStorageServerLocalCacheAbandonSizeKB = 512 * 1024; // abandon cache if >= 512 MB of it is invalid
+extern "C" int32 GStorageServerLocalCacheInvalidate = 0; // set to 1 to invalidate
+
 DEFINE_LOG_CATEGORY(LogStorageServerConnection);
 
 TRACE_DECLARE_INT_COUNTER(ZenHttpClientSerializedBytes, TEXT("ZenClient/SerializedBytes (compressed)"));
 TRACE_DECLARE_INT_COUNTER(ZenHttpClientThroughputBytes, TEXT("ZenClient/ThroughputBytes (decompressed)"));
 
+TRACE_DECLARE_FLOAT_COUNTER (ZenCacheRatio,     TEXT("ZenClient/Cache/HitRatio"));
+TRACE_DECLARE_INT_COUNTER   (ZenCacheHit,       TEXT("ZenClient/Cache/Hit"));
+TRACE_DECLARE_INT_COUNTER   (ZenCacheMiss,      TEXT("ZenClient/Cache/Miss"));
+TRACE_DECLARE_MEMORY_COUNTER(ZenCacheHitBytes,  TEXT("ZenClient/Cache/HitBytes"));
+TRACE_DECLARE_MEMORY_COUNTER(ZenCacheMissBytes, TEXT("ZenClient/Cache/MissBytes"));
+
+static void SetZenCacheRatio()
+{
+	const int64 HitCount = TRACE_COUNTER_GET(ZenCacheHit);
+	const int64 MissCount = TRACE_COUNTER_GET(ZenCacheMiss);
+	const float Ratio = (HitCount + MissCount > 0) ? ((float)HitCount / (HitCount + MissCount)) : 0.0f;
+	TRACE_COUNTER_SET(ZenCacheRatio, Ratio * 100.0f);
+}
+
+static void ZenCacheHit(const uint64 Bytes)
+{
+	TRACE_COUNTER_ADD(ZenCacheHit, 1);
+	TRACE_COUNTER_ADD(ZenCacheHitBytes, Bytes);
+	SetZenCacheRatio();
+}
+
+static void ZenCacheMiss(const uint64 Bytes)
+{
+	TRACE_COUNTER_ADD(ZenCacheMiss, 1);
+	TRACE_COUNTER_ADD(ZenCacheMissBytes, Bytes);
+	SetZenCacheRatio();
+}
+
 static FCbObject ProcessResponse(FIoBuffer IoBuffer, EStorageServerContentType ContentType)
 {
 	FCbObject ResponseObj;
@@ -63,6 +105,7 @@ bool FStorageServerConnection::Initialize(TArrayView<const FString> HostAddresse
 		if (HandshakeRequest())
 		{
 			UE_LOG(LogStorageServerConnection, Display, TEXT("Zen store connection established to %s:%i."), *CurrentHostAddr, Port);
+			SetupCacheStrategy();
 			return true;
 		}
 	}
@@ -132,10 +175,50 @@ void FStorageServerConnection::FileManifestRequest(TFunctionRef<void(FIoChunkId
 	}
 }
 
+void FStorageServerConnection::ChunkInfosRequest(TFunctionRef<void(FIoChunkId Id, FIoHash RawHash, int64 RawSize)> Callback)
+{
+	TRACE_CPUPROFILER_EVENT_SCOPE(FStorageServerConnection::ChunkInfosRequest);
+
+	TAnsiStringBuilder<256> ResourceBuilder;
+	ResourceBuilder.Append(BaseURI).Append("/chunkinfos?fieldnames=id,rawhash,rawsize");
+
+	IStorageServerHttpClient::FResult ResultTuple = HttpClient->RequestSync(*ResourceBuilder, EStorageServerContentType::CbObject);
+	TIoStatusOr<FIoBuffer> Result = ResultTuple.Get<0>();
+	if (Result.IsOk())
+	{
+		FMemoryReaderView Reader(Result.ValueOrDie().GetView());
+		FCbObject ResponseObj = LoadCompactBinary(Reader).AsObject();
+		for (FCbField& FileArrayEntry : ResponseObj["chunkinfos"].AsArray())
+		{
+			FCbObject Entry = FileArrayEntry.AsObject();
+			FCbObjectId Id = Entry["id"].AsObjectId();
+			FIoHash RawHash = Entry["rawhash"].AsHash();
+			int64 ResponseRawSize = Entry["rawsize"].AsInt64(-1);
+
+			FIoChunkId ChunkId;
+			ChunkId.Set(Id.GetView());
+
+			Callback(ChunkId, RawHash, ResponseRawSize);
+		}
+	}
+	else
+	{
+		UE_LOG(LogStorageServerConnection, Fatal, TEXT("Failed to read file manifest from storage server. '%s'"), *Result.Status().ToString());
+	}
+}
+
 int64 FStorageServerConnection::ChunkSizeRequest(const FIoChunkId& ChunkId)
 {
 	TRACE_CPUPROFILER_EVENT_SCOPE(FStorageServerConnection::ChunkSizeRequest);
 
+	int64 ChunkSize = 0;
+	if (CacheStrategy.IsValid() && CacheStrategy->TryGetChunkSize(ChunkId, ChunkSize))
+	{
+		return ChunkSize;
+	}
+
+	TRACE_CPUPROFILER_EVENT_SCOPE(FStorageServerConnection::ChunkSizeRequest::Http);
+
 	TAnsiStringBuilder<256> ResourceBuilder;
 	ResourceBuilder.Append(BaseURI);
 	ResourceBuilder << "/" << ChunkId << "/info";
@@ -150,7 +233,13 @@ int64 FStorageServerConnection::ChunkSizeRequest(const FIoChunkId& ChunkId)
 
 		FMemoryReaderView Reader(Result.ValueOrDie().GetView());
 		FCbObject ResponseObj = LoadCompactBinary(Reader).AsObject();
-		const int64 ChunkSize = ResponseObj["size"].AsInt64(0);
+		ChunkSize = ResponseObj["size"].AsInt64(0);
+
+		if (CacheStrategy.IsValid())
+		{
+			CacheStrategy->CacheChunkSize(ChunkId, ChunkSize);
+		}
+
 		return ChunkSize;
 	}
 	else if (Result.Status().GetErrorCode() != EIoErrorCode::NotFound)
@@ -173,14 +262,38 @@ TIoStatusOr<FIoBuffer> FStorageServerConnection::ReadChunkRequest(
 
 	const double StartTime = FPlatformTime::Seconds();
 
-	TAnsiStringBuilder<256> ResourceBuilder;
-	BuildReadChunkRequestUrl(ResourceBuilder, ChunkId, Offset, Size);
+	IStorageServerHttpClient::FResult ResultTuple;
 
-	IStorageServerHttpClient::FResult HttpResultTuple = HttpClient->RequestSync(*ResourceBuilder);
-	TIoStatusOr<FIoBuffer> ResultBuffer = ReadChunkRequestProcessHttpResult(HttpResultTuple, Offset, Size, OptDestination, bHardwareTargetBuffer);
+	bool bWasCached = false;
+	FIoBuffer CacheChunkBuffer;
+	EStorageServerContentType CacheContentType;
+	// TODO is there a way to pass destination directly?
+	if (CacheStrategy.IsValid() && CacheStrategy->ReadChunk(ChunkId, Offset, Size, TOptional<FIoBuffer>(), CacheChunkBuffer, CacheContentType))
+	{
+		bWasCached = true;
+		ZenCacheHit(CacheChunkBuffer.GetSize());
+		ResultTuple = IStorageServerHttpClient::FResult(CacheChunkBuffer, CacheContentType);
+	}
+	else
+	{
+		TRACE_CPUPROFILER_EVENT_SCOPE(FStorageServerConnection::ReadChunkRequest::Http);
+
+		TAnsiStringBuilder<256> ResourceBuilder;
+		BuildReadChunkRequestUrl(ResourceBuilder, ChunkId, Offset, Size);
+
+		ResultTuple = HttpClient->RequestSync(*ResourceBuilder);
+	}
+
+	TIoStatusOr<FIoBuffer> ResultBuffer = ReadChunkRequestProcessHttpResult(ResultTuple, Offset, Size, OptDestination, bHardwareTargetBuffer);
 
 	if (ResultBuffer.IsOk())
 	{
+		if (!bWasCached && CacheStrategy.IsValid())
+		{
+			ZenCacheMiss(ResultTuple.Key.ValueOrDie().GetSize());
+			CacheStrategy->CacheChunk(ChunkId, Offset, Size, ResultTuple.Key.ValueOrDie(), ResultTuple.Value);
+		}
+
 		const double Duration = FPlatformTime::Seconds() - StartTime;
 		AddTimingInstance(Duration, ResultBuffer.ValueOrDie().GetSize());
 	}
@@ -204,6 +317,8 @@ void FStorageServerConnection::ReadChunkRequestAsync(
 	TAnsiStringBuilder<256> ResourceBuilder;
 	BuildReadChunkRequestUrl(ResourceBuilder, ChunkId, Offset, Size);
 
+	// TODO use CacheStrategy
+
 	HttpClient->RequestAsync([this, Offset, Size, OptDestination, bHardwareTargetBuffer, OnResponse, StartTime](IStorageServerHttpClient::FResult HttpResultTuple)
 	{
 		TIoStatusOr<FIoBuffer> ResultBuffer = ReadChunkRequestProcessHttpResult(HttpResultTuple, Offset, Size, OptDestination, bHardwareTargetBuffer);
@@ -465,6 +580,53 @@ bool FStorageServerConnection::HandshakeRequest()
 	return false;
 }
 
+#if !HAS_STORAGE_SERVER_PLATFORM_CACHE_STRATEGY
+void FStorageServerConnection::SetupCacheStrategy()
+{
+	if (GStorageServerLocalCacheSizeKB <= 0)
+	{
+		CacheStrategy.Reset();
+		return;
+	}
+
+	FString DevStoragePath;
+#if defined(STORAGE_SERVER_PLATFORM_CACHE_ROOT_DIR)
+	DevStoragePath = STORAGE_SERVER_PLATFORM_CACHE_ROOT_DIR;
+#elif PLATFORM_ANDROID
+	DevStoragePath = FPlatformMisc::GamePersistentDownloadDir();
+#endif
+
+	const uint64 Size = ((uint64)GStorageServerLocalCacheSizeKB) * 1024;
+	const float FlushInterval = GStorageServerLocalCacheFlushInterval;
+	const uint64 FlushEveryNEntries = GStorageServerLocalCacheFlushEveryNEntries;
+	const uint64 AbandonSize = ((uint64)GStorageServerLocalCacheAbandonSizeKB) * 1024;
+	const bool bShouldInvalidate = GStorageServerLocalCacheInvalidate != 0;
+
+	TUniquePtr<StorageServer::ICacheJournal> Journal = MakeUnique<StorageServer::FCacheJournalSimple>(*(DevStoragePath / TEXT("ZenCache.db")), FlushEveryNEntries);
+	TUniquePtr<StorageServer::ICacheStorage> Storage = MakeUnique<StorageServer::FCacheStorageBlocking>(*(DevStoragePath / TEXT("ZenCache.storage")), Size);
+	CacheStrategy = MakeUnique<StorageServer::FCacheStrategyLinear>(MoveTemp(Journal), MoveTemp(Storage), AbandonSize, FlushInterval, bShouldInvalidate);
+
+	RefreshCacheEntries();
+}
+#endif // !HAS_STORAGE_SERVER_PLATFORM_CACHE_STRATEGY
+
+void FStorageServerConnection::RefreshCacheEntries()
+{
+	TRACE_CPUPROFILER_EVENT_SCOPE(StorageServerRefreshCacheEntries);
+
+	if (!CacheStrategy.IsValid())
+	{
+		return;
+	}
+
+	// TODO replace with a new zen server endpoint when it will become available
+	// TODO make it async
+	ChunkInfosRequest([&](FIoChunkId Id, FIoHash RawHash, int64 RawSize)
+	{
+		CacheStrategy->SetChunkInfo(Id, RawHash, RawSize); // using RawHash as a subtitute for ModHash.
+	});
+}
+
 void FStorageServerConnection::BuildReadChunkRequestUrl(FAnsiStringBuilderBase& Builder, const FIoChunkId& ChunkId, const uint64 Offset, const uint64 Size)
 {
 	Builder.Append(BaseURI) << "/" << ChunkId;
diff --git a/Engine/Source/Runtime/StorageServerClient/Private/StorageServerConnection.h b/Engine/Source/Runtime/StorageServerClient/Private/StorageServerConnection.h
--- a/Engine/Source/Runtime/StorageServerClient/Private/StorageServerConnection.h
+++ b/Engine/Source/Runtime/StorageServerClient/Private/StorageServerConnection.h
@@ -13,6 +13,15 @@
 #include "HAL/PlatformTime.h"
 #include "StorageServerHttpClient.h"
 #include "IStorageServerPlatformFile.h"
+#include "Cache/CacheStrategy.h"
+
+#ifndef PLATFORM_SUPPORTS_STORAGE_SERVER_CACHE
+	#define PLATFORM_SUPPORTS_STORAGE_SERVER_CACHE (PLATFORM_WINDOWS || PLATFORM_ANDROID)
+#endif
+
+#ifndef HAS_STORAGE_SERVER_PLATFORM_CACHE_STRATEGY
+	#define HAS_STORAGE_SERVER_PLATFORM_CACHE_STRATEGY 0
+#endif
 
 #if !UE_BUILD_SHIPPING
 
@@ -31,6 +40,7 @@ public:
 
 	void PackageStoreRequest(TFunctionRef<void(FPackageStoreEntryResource&&)> Callback);
 	void FileManifestRequest(TFunctionRef<void(FIoChunkId Id, FStringView Path, int64 RawSize)> Callback);
+	void ChunkInfosRequest(TFunctionRef<void(FIoChunkId Id, FIoHash RawHash, int64 RawSize)> Callback);
 	int64 ChunkSizeRequest(const FIoChunkId& ChunkId);
 	TIoStatusOr<FIoBuffer> ReadChunkRequest(
 		const FIoChunkId& ChunkId,
@@ -57,6 +67,7 @@ public:
 
 private:
 	TUniquePtr<IStorageServerHttpClient> HttpClient;
+	TUniquePtr<StorageServer::ICacheStrategy> CacheStrategy;
 	FAnsiString BaseURI;
 	FString CurrentHostAddr;
 
@@ -71,6 +82,8 @@ private:
 	TUniquePtr<IStorageServerHttpClient> CreateHttpClient(const FString Address, const int32 Port);
 	TSharedPtr<FInternetAddr> StringToInternetAddr(const FString Address, const int32 Port);
 	bool HandshakeRequest();
+	void SetupCacheStrategy();
+	void RefreshCacheEntries();
 	void BuildReadChunkRequestUrl(FAnsiStringBuilderBase& Builder, const FIoChunkId& ChunkId, const uint64 Offset, const uint64 Size);
 	static TIoStatusOr<FIoBuffer> ReadChunkRequestProcessHttpResult(
 		IStorageServerHttpClient::FResult ResultTuple,
