Disable zen client side caching

diff --git a/Engine/Source/Runtime/StorageServerClient/Private/StorageServerConnection.cpp b/Engine/Source/Runtime/StorageServerClient/Private/StorageServerConnection.cpp
--- a/Engine/Source/Runtime/StorageServerClient/Private/StorageServerConnection.cpp
+++ b/Engine/Source/Runtime/StorageServerClient/Private/StorageServerConnection.cpp
@@ -29,7 +29,7 @@
 #if !UE_BUILD_SHIPPING
 
 // Set to true to enable cache, temporarily disabled for now
-static constexpr bool bEnableZenCache = true; // PLATFORM_SUPPORTS_STORAGE_SERVER_CACHE;
+static constexpr bool bEnableZenCache = false; // PLATFORM_SUPPORTS_STORAGE_SERVER_CACHE;
 
 // TODO implement a way to control cache via cvars, right now it's not possible as cvars are initialized at later stage than cache
 extern "C" int32 GStorageServerLocalCacheSizeKB = bEnableZenCache ? 8 * 1024 * 1024 : 0; // 8 GB
