Improve failed to connect to Zen server display message to ask user to check IP address in ue.projectstore

diff --git a/Engine/Source/Runtime/StorageServerClient/Private/StorageServerPlatformFile.cpp b/Engine/Source/Runtime/StorageServerClient/Private/StorageServerPlatformFile.cpp
--- a/Engine/Source/Runtime/StorageServerClient/Private/StorageServerPlatformFile.cpp
+++ b/Engine/Source/Runtime/StorageServerClient/Private/StorageServerPlatformFile.cpp
@@ -1103,7 +1103,8 @@ void FStorageServerPlatformFile::InitializeConnection()
 				"Network data streaming failed to connect to any of the following data sources:\n\n%s\n\n"
 				"This can be due to the sources being offline, the Unreal Zen Storage process not currently running, "
 				"invalid addresses, firewall blocking, or the sources being on a different network from this device.\n"
-				"Please verify that your Unreal Zen Storage process is running using the ZenDashboard utility. "
+				"Please verify that your Unreal Zen Storage process is running using the ZenDashboard utility, "
+				"and ue.projectstore file in the staged folder contains the valid IP address of the host PC in the \"remotehostnames\" section.\n"
 				"If these issues can't be addressed, you can use an installed build without network data streaming by "
 				"building with the '-pak' argument. This process will now exit."),
 				*FString::Join(HostAddrs, TEXT("\n")));
