[Trivial] Rename the experimental StorageServerLocalCache with the more accurate name StorageServerIASCache.
- this class is likely to be removed soon in favor of the new ICacheStrategy classes.

#rnx

diff --git a/Engine/Source/Runtime/StorageServerClient/Private/StorageServerIASCache.cpp b/Engine/Source/Runtime/StorageServerClient/Private/StorageServerIASCache.cpp
new file mode 100644
--- /dev/null
+++ b/Engine/Source/Runtime/StorageServerClient/Private/StorageServerIASCache.cpp
@@ -0,0 +1,291 @@
+// Copyright Epic Games, Inc. All Rights Reserved.
+
+#include "StorageServerIASCache.h"
+#include "StorageServerConnection.h"
+#include "Misc/ScopeLock.h"
+#include "Misc/Paths.h"
+#include "Misc/CoreDelegates.h"
+#include "Misc/CommandLine.h"
+#include "ProfilingDebugging/IoStoreTrace.h"
+#include "ProfilingDebugging/CountersTrace.h"
+#include "Modules/ModuleManager.h"
+
+#if !UE_BUILD_SHIPPING && WITH_STORAGESERVER_IAS_CACHE
+
+/******************************************************************************************************************************
+*                                                                                                                             * 
+*  Enable with -ZenIASCache  + can use  -ZenDropIASCache to clear the IAS cache on startup.								  *
+*                                                                                                                             * 
+*******************************************************************************************************************************/
+
+int32 GStorageServerIASCacheMinDataSizeKB = 0;
+static FAutoConsoleVariableRef CVar_StorageServerIASCacheMinDataSizeKB(
+	TEXT("s.StorageServerIASCache.MinDataSizeKB"),
+	GStorageServerIASCacheMinDataSizeKB,
+	TEXT("Minimum data size to store in the cache (KB)")
+);
+
+int32 GStorageServerIASCacheMaxDataSizeKB = -1;
+static FAutoConsoleVariableRef CVar_StorageServerIASCacheMaxDataSizeKB(
+	TEXT("s.StorageServerIASCache.MaxDataSizeKB"),
+	GStorageServerIASCacheMaxDataSizeKB,
+	TEXT("Maximum data size to store in the cache (KB). Set to -1 to try and cache everything")
+);
+
+namespace
+{
+	// this is a way to optionally ignore a certain proportion of items from the cache (based on ChunkId hash - not perfect)
+	int32 GCachePutMask = 0; //0 -> put all,  1 -> put 1/2, 3 -> put 1/4, 7 -> put 1/8, 15 -> put 1/16
+
+	// disable all cache writes (for testing/debugging)
+	bool bCacheIsReadOnly = false;
+}
+
+TRACE_DECLARE_INT_COUNTER(ZenIASCacheHitBytes,         TEXT("ZenClient/IASCache/MemoryHitBytes"));
+TRACE_DECLARE_INT_COUNTER(ZenIASCacheMaterializeBytes, TEXT("ZenClient/IASCache/DiskHitBytes"));
+TRACE_DECLARE_INT_COUNTER(ZenIASCacheMissBytes,        TEXT("ZenClient/IASCache/MissBytes"));
+TRACE_DECLARE_INT_COUNTER(ZenIASCachePutBytes,         TEXT("ZenClient/IASCache/PutBytes"));
+
+/******************************************************************************************************************************/
+
+
+
+
+FStorageServerIASCache::FStorageServerIASCache(FStorageServerConnection& InConnection) 
+	: Connection(InConnection)
+{
+	if (FModuleManager::Get().ModuleExists(TEXT("IoStoreOnDemand")) )
+	{
+		FCoreDelegates::OnInit.AddLambda([this]()
+		{
+				if (FPaths::HasProjectPersistentDownloadDir())
+				{
+					FString CacheDir = FPaths::ProjectPersistentDownloadDir();
+					if (!CacheDir.IsEmpty())
+					{
+						TRACE_CPUPROFILER_EVENT_SCOPE(StorageServerIoDispatcherInitIASCache);
+
+						UE::IoStore::FIasCacheConfig Config;
+						Config.Name = TEXT("zencache");
+						Config.DropCache = FParse::Param(FCommandLine::Get(), TEXT("ZenDropIASCache"));
+						Cache = UE::IoStore::MakeIasCache(*CacheDir, Config);
+
+						TRACE_IOSTORE_BACKEND_NAME(this, TEXT("StorageServerIASCache") );
+				
+						FIoDispatcher::Get().Mount(AsShared(), 1); //must be higher priority than FStorageServerIoDispatcherBackend
+					}
+				}
+		});
+	}
+}
+
+void FStorageServerIASCache::Initialize(TSharedRef<const FIoDispatcherBackendContext> Context) 
+{
+	BackendContext = Context; 
+}
+
+
+FIoHash FStorageServerIASCache::GetRequestCacheKey(FIoRequestImpl* Request) const
+{
+	FIoOffsetAndLength Range(Request->Options.GetOffset(), Request->Options.GetSize());
+
+	FIoHashBuilder HashBuilder;
+	HashBuilder.Update(Request->ChunkId.GetData(), Request->ChunkId.GetSize());
+	HashBuilder.Update(&Range, sizeof(FIoOffsetAndLength));
+	return HashBuilder.Finalize();
+}
+
+bool FStorageServerIASCache::ResolveFromCache(FIoRequestImpl* Request)
+{
+	// bail out if there's no cache
+	if (!Cache.IsValid())
+	{
+		return false;
+	}
+
+	// bail out if we know isn't something that would be in the cache
+	uint32 ChunkHash = GetTypeHash(Request->ChunkId);
+	if (GCachePutMask > 0 && ((ChunkHash & GCachePutMask) != GCachePutMask))
+	{
+		return false;
+	}
+
+	// see if the chunk is in the cache
+	FIoHash Key = GetRequestCacheKey(Request);
+	if (!Cache->ContainsChunk(Key))
+	{
+		return false;
+	}
+
+	// helper lambda to set the result buffer
+	auto SetRequestResult = [Request, this](FIoBuffer& InBuffer)
+	{
+		if (Request->Options.GetTargetVa())
+		{
+			FIoBuffer VaBuffer(FIoBuffer::Wrap, Request->Options.GetTargetVa(), Request->Options.GetSize());
+			FMemory::Memcpy(VaBuffer.GetData(), InBuffer.GetData(), InBuffer.GetSize());
+			Request->SetResult(VaBuffer);
+
+		}
+		else
+		{
+			Request->SetResult(InBuffer);
+		}
+	};
+
+	// start processing - first, check the memory cache
+	EIoErrorCode ErrorCode = EIoErrorCode::Unknown;
+	{
+		FIoBuffer Buffer;
+		ErrorCode = Cache->Get(Key, Buffer);
+		if (ErrorCode == EIoErrorCode::Ok)
+		{
+			TRACE_IOSTORE_BACKEND_REQUEST_STARTED(Request, this);
+			SetRequestResult(Buffer);
+			CompletedRequests.Enqueue(Request);
+			BackendContext->WakeUpDispatcherThreadDelegate.Execute();
+			TRACE_COUNTER_ADD(ZenIASCacheHitBytes, Buffer.GetSize());
+			TRACE_IOSTORE_BACKEND_REQUEST_COMPLETED(Request, Buffer.GetSize());
+			return true;
+		}
+	}
+
+	// FileNotOpen means it is cached on disk - materialize it now
+	if (ErrorCode == EIoErrorCode::FileNotOpen)
+	{
+		TRACE_IOSTORE_BACKEND_REQUEST_STARTED(Request, this);
+
+		struct FPendingMaterialize
+		{
+			EIoErrorCode ErrorCode;
+			FIoBuffer Buffer;
+		};
+		FPendingMaterialize* PendingMaterialize = new FPendingMaterialize();
+		
+		UE::Tasks::FTaskEvent MaterializeDone(TEXT("ZenIasCacheMaterializeDone"));
+		UE::Tasks::FTask TaskHandle = UE::Tasks::Launch(UE_SOURCE_LOCATION, [this, Request, PendingMaterialize, SetRequestResult = MoveTemp(SetRequestResult)]
+		{
+			// complete the request
+			if (PendingMaterialize->ErrorCode == EIoErrorCode::Ok)
+			{
+				SetRequestResult(PendingMaterialize->Buffer);
+				TRACE_COUNTER_ADD(ZenIASCacheMaterializeBytes, PendingMaterialize->Buffer.GetSize());
+				TRACE_IOSTORE_BACKEND_REQUEST_COMPLETED(Request, PendingMaterialize->Buffer.GetSize());
+			}
+			else
+			{
+				Request->SetFailed();
+				TRACE_IOSTORE_BACKEND_REQUEST_FAILED(Request);
+			}
+			CompletedRequests.Enqueue(Request);
+			BackendContext->WakeUpDispatcherThreadDelegate.Execute();
+
+			//clean up
+			delete PendingMaterialize;
+		}, MaterializeDone, UE::Tasks::ETaskPriority::Normal);
+
+		Cache->Materialize(Key, PendingMaterialize->Buffer, PendingMaterialize->ErrorCode, MaterializeDone);
+		return true;
+	}
+
+	// can't process this request
+	return false;
+}
+
+void FStorageServerIASCache::OptionalCacheWrite(FIoRequestImpl* Request, const FIoBuffer& Buffer)
+{
+	if (Cache.IsValid())
+	{
+		uint64 BufferSize = Buffer.GetSize();
+		TRACE_COUNTER_ADD(ZenIASCacheMissBytes, BufferSize);
+
+		// bail out if it's read only (for testing only!)
+		if (bCacheIsReadOnly)
+		{
+			return;
+		}
+
+
+		// bail out if this isn't something we're going to be caching
+		uint32 ChunkHash = GetTypeHash(Request->ChunkId);
+		if (GCachePutMask > 0 && ((ChunkHash & GCachePutMask) != GCachePutMask))
+		{
+			return;
+		}
+
+		// bail out if it's too small
+		if (BufferSize < GStorageServerIASCacheMinDataSizeKB*1024)
+		{
+			return;
+		}
+
+		// bail out if it's too large
+		if (GStorageServerIASCacheMaxDataSizeKB > 0 && BufferSize > GStorageServerIASCacheMaxDataSizeKB*1024)
+		{
+			return;
+		}
+
+		// take a copy of the buffer and store it in the cache
+		FIoBuffer PutBuffer(Buffer.GetSize());
+		FMemory::Memcpy(PutBuffer.GetData(), Buffer.GetData(), Buffer.GetSize());
+
+		TRACE_COUNTER_ADD(ZenIASCachePutBytes, BufferSize);
+		Cache->Put(GetRequestCacheKey(Request), PutBuffer);
+	}
+}
+
+void FStorageServerIASCache::ResolveIoRequests(FIoRequestList Requests, FIoRequestList& OutUnresolved)
+{
+	if (Cache.IsValid() && BackendContext->bIsMultiThreaded)
+	{
+		while (FIoRequestImpl* Request = Requests.PopHead())
+		{
+			if (ResolveFromCache(Request) == false)
+			{
+				OutUnresolved.AddTail(Request);
+			}
+		}
+	}
+}
+
+FIoRequestImpl* FStorageServerIASCache::GetCompletedIoRequests()
+{
+	FIoRequestImpl* Result = nullptr;
+
+	FIoRequestImpl* Item = nullptr;
+	while (CompletedRequests.Dequeue(Item))
+	{
+		Item->NextRequest = Result;
+		Result = Item;
+	}
+
+	return Result;
+}
+
+bool FStorageServerIASCache::DoesChunkExist(const FIoChunkId& ChunkId) const
+{
+	return GetSizeForChunk(ChunkId).IsOk();
+}
+
+TIoStatusOr<uint64> FStorageServerIASCache::GetSizeForChunk(const FIoChunkId& ChunkId) const
+{
+	TRACE_CPUPROFILER_EVENT_SCOPE(StorageServerIoDispatcherIASCacheGetSizeForChunk);
+	int64 ChunkSize = Connection.ChunkSizeRequest(ChunkId);
+	if (ChunkSize >= 0)
+	{
+		return ChunkSize;
+	}
+	else
+	{
+		return FIoStatus(EIoErrorCode::NotFound);
+	}
+}
+
+const TCHAR* FStorageServerIASCache::GetName() const
+{
+	return TEXT("StorageServerIASCache");
+}
+
+
+#endif //!UE_BUILD_SHIPPING
+
diff --git a/Engine/Source/Runtime/StorageServerClient/Private/StorageServerIASCache.h b/Engine/Source/Runtime/StorageServerClient/Private/StorageServerIASCache.h
new file mode 100644
--- /dev/null
+++ b/Engine/Source/Runtime/StorageServerClient/Private/StorageServerIASCache.h
@@ -0,0 +1,50 @@
+// Copyright Epic Games, Inc. All Rights Reserved.
+
+#pragma once
+
+#include "IO/IoDispatcherBackend.h"
+
+#ifndef WITH_STORAGESERVER_IAS_CACHE
+	#define WITH_STORAGESERVER_IAS_CACHE 0 // disabled by default. this code is likely to be removed soon
+#endif
+
+#if !UE_BUILD_SHIPPING && WITH_STORAGESERVER_IAS_CACHE
+#include "IO/IasCache.h"
+#include "Containers/MpscQueue.h"
+
+class FStorageServerConnection;
+
+class FStorageServerIASCache : public IIoDispatcherBackend, public TSharedFromThis<FStorageServerIASCache>
+{
+public:
+	FStorageServerIASCache(FStorageServerConnection& InConnection);
+	virtual ~FStorageServerIASCache() = default;
+
+	bool ResolveFromCache(FIoRequestImpl* Request);
+	void OptionalCacheWrite(FIoRequestImpl* Request, const FIoBuffer& Buffer);
+
+private:
+	FIoHash GetRequestCacheKey(FIoRequestImpl* Request) const;
+
+	FStorageServerConnection& Connection;
+	TUniquePtr<UE::IoStore::IIasCache> Cache;
+	TMpscQueue<FIoRequestImpl*> CompletedRequests;
+	TSharedPtr<const FIoDispatcherBackendContext> BackendContext;
+
+
+	// IIoDispatcherBackend implementation
+	virtual void Initialize(TSharedRef<const FIoDispatcherBackendContext> Context) override;
+	virtual void ResolveIoRequests(FIoRequestList Requests, FIoRequestList& OutUnresolved)  override;
+	virtual FIoRequestImpl* GetCompletedIoRequests()  override;
+	virtual void CancelIoRequest(FIoRequestImpl* Request)  override {}
+	virtual void UpdatePriorityForIoRequest(FIoRequestImpl* Request)  override {}
+	virtual bool DoesChunkExist(const FIoChunkId& ChunkId) const  override;
+	virtual TIoStatusOr<uint64> GetSizeForChunk(const FIoChunkId& ChunkId) const override;
+	virtual TIoStatusOr<FIoMappedRegion> OpenMapped(const FIoChunkId& ChunkId, const FIoReadOptions& Options) override
+	{ 
+		return FIoStatus(EIoErrorCode::NotFound); 
+	}
+	virtual const TCHAR* GetName() const override;
+};
+
+#endif
diff --git a/Engine/Source/Runtime/StorageServerClient/Private/StorageServerIoDispatcherBackend.cpp b/Engine/Source/Runtime/StorageServerClient/Private/StorageServerIoDispatcherBackend.cpp
--- a/Engine/Source/Runtime/StorageServerClient/Private/StorageServerIoDispatcherBackend.cpp
+++ b/Engine/Source/Runtime/StorageServerClient/Private/StorageServerIoDispatcherBackend.cpp
@@ -29,11 +29,11 @@ FStorageServerIoDispatcherBackend::FStorageServerIoDispatcherBackend(FStorageSer
 	, NewRequestEvent(FPlatformProcess::GetSynchEventFromPool(false))
 	, BatchCompletedEvent(FPlatformProcess::GetSynchEventFromPool(false))
 {
-#if WITH_STORAGESERVER_LOCAL_CACHE
-	bool bUseLocalCache = FParse::Param(FCommandLine::Get(), TEXT("ZenLocalCache"));
-	if (bUseLocalCache)
+#if WITH_STORAGESERVER_IAS_CACHE
+	bool bUseIASCache = FParse::Param(FCommandLine::Get(), TEXT("ZenIASCache"));
+	if (bUseIASCache)
 	{
-		Cache = MakeShared<FStorageServerLocalCache>(Connection);
+		Cache = MakeShared<FStorageServerIASCache>(Connection);
 	}
 #endif
 }
@@ -164,7 +164,7 @@ bool FStorageServerIoDispatcherBackend::Resolve(FIoRequestImpl* Request)
 		{
 			Request->SetResult(Result.ValueOrDie());
 			TRACE_IOSTORE_BACKEND_REQUEST_COMPLETED(Request, Result.ValueOrDie().GetSize());
-#if WITH_STORAGESERVER_LOCAL_CACHE
+#if WITH_STORAGESERVER_IAS_CACHE
 			if (Cache.IsValid())
 			{
 				Cache->OptionalCacheWrite(Request, Result.ValueOrDie());
@@ -367,7 +367,7 @@ void FStorageServerIoDispatcherBackend::FBatch::DoThreadedWork()
 			Request->SetResult(Result.ValueOrDie());
 			TRACE_IOSTORE_BACKEND_REQUEST_COMPLETED(Request, Result.ValueOrDie().GetSize());
 
-#if WITH_STORAGESERVER_LOCAL_CACHE
+#if WITH_STORAGESERVER_IAS_CACHE
 			if (Owner.Cache.IsValid())
 			{
 				Owner.Cache->OptionalCacheWrite(Request, Result.ValueOrDie());
diff --git a/Engine/Source/Runtime/StorageServerClient/Private/StorageServerIoDispatcherBackend.h b/Engine/Source/Runtime/StorageServerClient/Private/StorageServerIoDispatcherBackend.h
--- a/Engine/Source/Runtime/StorageServerClient/Private/StorageServerIoDispatcherBackend.h
+++ b/Engine/Source/Runtime/StorageServerClient/Private/StorageServerIoDispatcherBackend.h
@@ -3,7 +3,7 @@
 #pragma once
 
 #include "IO/IoDispatcherBackend.h"
-#include "StorageServerLocalCache.h"
+#include "StorageServerIASCache.h"
 #include "HAL/Runnable.h"
 #include "Misc/IQueuedWork.h"
 
@@ -99,8 +99,8 @@ private:
 	FIoRequestImpl* CompletedRequestsHead = nullptr;
 	FIoRequestImpl* CompletedRequestsTail = nullptr;
 	TAtomic<bool> bStopRequested{ false };
-#if WITH_STORAGESERVER_LOCAL_CACHE
-	TSharedPtr<FStorageServerLocalCache> Cache;
+#if WITH_STORAGESERVER_IAS_CACHE
+	TSharedPtr<FStorageServerIASCache> Cache;
 #endif
 };
 
diff --git a/Engine/Source/Runtime/StorageServerClient/Private/StorageServerLocalCache.cpp b/Engine/Source/Runtime/StorageServerClient/Private/StorageServerLocalCache.cpp
deleted file mode 100644
diff --git a/Engine/Source/Runtime/StorageServerClient/Private/StorageServerLocalCache.h b/Engine/Source/Runtime/StorageServerClient/Private/StorageServerLocalCache.h
deleted file mode 100644
