Zen client-side cache queries the server's chunk list asyncronously.

#rnx

diff --git a/Engine/Source/Runtime/StorageServerClient/Private/Cache/CacheJournal.h b/Engine/Source/Runtime/StorageServerClient/Private/Cache/CacheJournal.h
--- a/Engine/Source/Runtime/StorageServerClient/Private/Cache/CacheJournal.h
+++ b/Engine/Source/Runtime/StorageServerClient/Private/Cache/CacheJournal.h
@@ -1,10 +1,11 @@
-﻿// Copyright Epic Games, Inc. All Rights Reserved.
+// Copyright Epic Games, Inc. All Rights Reserved.
 
 #pragma once
 
 #include "IO/IoChunkId.h"
 #include "IO/IoHash.h"
 #include "Templates/Function.h"
+#include "Containers/Array.h"
 #include "Math/Interval.h"
 #include "StorageServerHttpClient.h"
 
@@ -16,6 +17,13 @@ namespace StorageServer
 	// right now is implemented via ChunkId RawHash as a stop gap solution
 	typedef FIoHash FModHash;
 
+	struct FLatestServerChunkInfo
+	{
+		FIoChunkId ChunkId;
+		FModHash ModHash;
+		int64 RawSize;
+	};
+
 	struct FCacheChunkInfo
 	{
 		FModHash ModHash;   // content modification chunk hash
@@ -147,6 +155,12 @@ namespace StorageServer
 		virtual void IterateCacheEntries(
 			TFunctionRef<void(const FIoChunkId& ChunkId, const FCacheEntry& Entry)> Callback
 		) = 0;
+
+		// Updates a group of chunk infos based on the latest from the server. Expected to be used during initialization
+		// Returns all invalidated chunk Ids
+		virtual TArray<FIoChunkId> ApplyLatestServerChunkInfo( 
+			const TArray<FLatestServerChunkInfo>& ServerChunkInfoArray 
+		) = 0;
 	};
 }
 
diff --git a/Engine/Source/Runtime/StorageServerClient/Private/Cache/CacheJournalSimple.cpp b/Engine/Source/Runtime/StorageServerClient/Private/Cache/CacheJournalSimple.cpp
--- a/Engine/Source/Runtime/StorageServerClient/Private/Cache/CacheJournalSimple.cpp
+++ b/Engine/Source/Runtime/StorageServerClient/Private/Cache/CacheJournalSimple.cpp
@@ -1,4 +1,4 @@
-﻿// Copyright Epic Games, Inc. All Rights Reserved.
+// Copyright Epic Games, Inc. All Rights Reserved.
 
 #include "CacheJournalSimple.h"
 
@@ -16,6 +16,8 @@ namespace StorageServer
 {
 	FCacheJournalSimple::FCacheJournalSimple(const TCHAR* InFileName, const uint64 InFlushAtWriteCount)
 	{
+		TRACE_CPUPROFILER_EVENT_SCOPE(FCacheJournalSimple::Initialize);
+
 		FileName = InFileName;
 		FlushAtWriteCount = InFlushAtWriteCount;
 		CurrentWriteCount = 0;
@@ -146,6 +148,37 @@ namespace StorageServer
 		return true;
 	}
 
+	TArray<FIoChunkId> FCacheJournalSimple::ApplyLatestServerChunkInfo(const TArray<FLatestServerChunkInfo>& ServerChunkInfoArray)
+	{
+		TRACE_CPUPROFILER_EVENT_SCOPE(FCacheJournalSimple::ApplyLatestServerChunkInfo);
+
+		TArray<FIoChunkId> InvalidatedChunks;
+		for (const FLatestServerChunkInfo& ServerChunkInfo : ServerChunkInfoArray)
+		{
+			FScopeLock Lock(&DataLock);
+			FCacheChunkInfo* ChunkInfo = ChunkInfos.Find(ServerChunkInfo.ChunkId);
+			if (ChunkInfo)
+			{
+				FModHash PreviousHash = ChunkInfo->ModHash;
+				const bool bHashesMatch = ChunkInfo->SetChunkInfo(&ServerChunkInfo.ModHash, &ServerChunkInfo.RawSize, nullptr);
+				if (!bHashesMatch)
+				{
+					TRACE_COUNTER_INCREMENT(ZenJournalSimpleInvalidated);
+					InvalidatedChunks.Add(ServerChunkInfo.ChunkId);
+				}
+			}
+			else
+			{
+				FCacheChunkInfo NewChunkInfo = {};
+				NewChunkInfo.SetChunkInfo(&ServerChunkInfo.ModHash, &ServerChunkInfo.RawSize, nullptr);
+				ChunkInfos.Add(ServerChunkInfo.ChunkId, NewChunkInfo);
+				bDirty = true;
+			}
+		}
+
+		return MoveTemp(InvalidatedChunks);
+	}
+
 	bool FCacheJournalSimple::TryGetChunkInfo(const FIoChunkId& ChunkId, FCacheChunkInfo& OutChunkInfo)
 	{
 		FScopeLock Lock(&DataLock);
diff --git a/Engine/Source/Runtime/StorageServerClient/Private/Cache/CacheJournalSimple.h b/Engine/Source/Runtime/StorageServerClient/Private/Cache/CacheJournalSimple.h
--- a/Engine/Source/Runtime/StorageServerClient/Private/Cache/CacheJournalSimple.h
+++ b/Engine/Source/Runtime/StorageServerClient/Private/Cache/CacheJournalSimple.h
@@ -1,4 +1,4 @@
-﻿// Copyright Epic Games, Inc. All Rights Reserved.
+// Copyright Epic Games, Inc. All Rights Reserved.
 
 #pragma once
 
@@ -56,6 +56,10 @@ namespace StorageServer
 		virtual void IterateCacheEntries(
 			TFunctionRef<void(const FIoChunkId& ChunkId, const FCacheEntry& Entry)> Callback
 		) override;
+
+		virtual TArray<FIoChunkId> ApplyLatestServerChunkInfo(
+			const TArray<FLatestServerChunkInfo>& ServerChunkInfoArray 
+		) override;
 		
 	private:
 		FCriticalSection DataLock, FlushTaskLock;
diff --git a/Engine/Source/Runtime/StorageServerClient/Private/Cache/CacheStrategy.h b/Engine/Source/Runtime/StorageServerClient/Private/Cache/CacheStrategy.h
--- a/Engine/Source/Runtime/StorageServerClient/Private/Cache/CacheStrategy.h
+++ b/Engine/Source/Runtime/StorageServerClient/Private/Cache/CacheStrategy.h
@@ -1,4 +1,4 @@
-﻿// Copyright Epic Games, Inc. All Rights Reserved.
+// Copyright Epic Games, Inc. All Rights Reserved.
 
 #pragma once
 
@@ -103,6 +103,16 @@ namespace StorageServer
 			Storage->Flush();
 		}
 
+		virtual void ApplyLatestServerChunkInfo(const TArray<FLatestServerChunkInfo>& ServerChunkInfoArray)
+		{
+			// naive implementation - specializations can likely do better
+			TArray<FIoChunkId> InvalidatedChunks = Journal->ApplyLatestServerChunkInfo(ServerChunkInfoArray);
+			for (const FIoChunkId& ChunkId : InvalidatedChunks)
+			{
+				Invalidate(ChunkId);
+			}
+		}
+
 	protected:
 		TUniquePtr<ICacheJournal> Journal;
 		TUniquePtr<ICacheStorage> Storage;
diff --git a/Engine/Source/Runtime/StorageServerClient/Private/Cache/CacheStrategyLinear.cpp b/Engine/Source/Runtime/StorageServerClient/Private/Cache/CacheStrategyLinear.cpp
--- a/Engine/Source/Runtime/StorageServerClient/Private/Cache/CacheStrategyLinear.cpp
+++ b/Engine/Source/Runtime/StorageServerClient/Private/Cache/CacheStrategyLinear.cpp
@@ -83,6 +83,35 @@ namespace StorageServer
 		}
 	}
 
+	void FCacheStrategyLinear::ApplyLatestServerChunkInfo(const TArray<FLatestServerChunkInfo>& ServerChunkInfoArray)
+	{
+		// update the journal
+		TArray<FIoChunkId> InvalidatedChunks = Journal->ApplyLatestServerChunkInfo(ServerChunkInfoArray);
+
+		// invalidate affected items
+		if (CurrentSize > 0)
+		{
+			TRACE_CPUPROFILER_EVENT_SCOPE(FCacheStrategyLinear::Invalidate);
+
+			for (const FIoChunkId& ChunkId : InvalidatedChunks)
+			{
+				Journal->IterateCacheEntriesForChunkId(ChunkId, [this](const FCacheEntry& CacheEntry)
+				{
+					CurrentInvalidSize += CacheEntry.StorageSize;
+				});
+
+				if (CurrentInvalidSize >= AbandonAtInvalidSize)
+				{
+					SetCounters();
+					InvalidateAll();
+					return;
+				}
+			}
+
+			SetCounters();
+		}
+	}
+
 	bool FCacheStrategyLinear::ReadChunk(const FIoChunkId& RequestChunkId, const uint64 RequestOffset, const uint64 RequestSize, TOptional<FIoBuffer> OptDestination, FIoBuffer& OutBuffer, EStorageServerContentType& OutContentType)
 	{
 		FCacheEntry Entry = {};
diff --git a/Engine/Source/Runtime/StorageServerClient/Private/Cache/CacheStrategyLinear.h b/Engine/Source/Runtime/StorageServerClient/Private/Cache/CacheStrategyLinear.h
--- a/Engine/Source/Runtime/StorageServerClient/Private/Cache/CacheStrategyLinear.h
+++ b/Engine/Source/Runtime/StorageServerClient/Private/Cache/CacheStrategyLinear.h
@@ -1,4 +1,4 @@
-﻿// Copyright Epic Games, Inc. All Rights Reserved.
+// Copyright Epic Games, Inc. All Rights Reserved.
 
 #pragma once
 
@@ -23,6 +23,8 @@ namespace StorageServer
 
 		virtual void CacheChunk(const FIoChunkId& RequestChunkId, const uint64 RequestOffset, const uint64 RequestSize, const FIoBuffer& Buffer, const EStorageServerContentType ContentType) override;
 
+		virtual void ApplyLatestServerChunkInfo(const TArray<FLatestServerChunkInfo>& ServerChunkInfoArray) override;
+		
 	private:
 		// Linear cache is implemented as arena allocator
 		std::atomic<uint64> CurrentSize; // Current size of arena allocator
diff --git a/Engine/Source/Runtime/StorageServerClient/Private/StorageServerConnection.cpp b/Engine/Source/Runtime/StorageServerClient/Private/StorageServerConnection.cpp
--- a/Engine/Source/Runtime/StorageServerClient/Private/StorageServerConnection.cpp
+++ b/Engine/Source/Runtime/StorageServerClient/Private/StorageServerConnection.cpp
@@ -21,6 +21,8 @@
 #include "BuiltInHttpClient/BuiltInHttpClientFSocket.h"
 #include "BuiltInHttpClient/BuiltInHttpClientPlatformSocket.h"
 #include "HAL/PlatformMath.h"
+#include "HAL/RunnableThread.h"
+#include "HAL/Event.h"
 
 #include "Cache/CacheJournalSimple.h"
 #include "Cache/CacheStorageBlocking.h"
@@ -211,8 +213,10 @@ int64 FStorageServerConnection::ChunkSizeRequest(const FIoChunkId& ChunkId)
 {
 	TRACE_CPUPROFILER_EVENT_SCOPE(FStorageServerConnection::ChunkSizeRequest);
 
+	bool bCacheAvailable = FinalizeSetupCacheStrategy();
+
 	int64 ChunkSize = 0;
-	if (CacheStrategy.IsValid() && CacheStrategy->TryGetChunkSize(ChunkId, ChunkSize))
+	if (bCacheAvailable && CacheStrategy.IsValid() && CacheStrategy->TryGetChunkSize(ChunkId, ChunkSize))
 	{
 		return ChunkSize;
 	}
@@ -260,6 +264,8 @@ TIoStatusOr<FIoBuffer> FStorageServerConnection::ReadChunkRequest(
 {
 	TRACE_CPUPROFILER_EVENT_SCOPE(FStorageServerConnection::ReadChunkRequest);
 
+	bool bCacheAvailable = FinalizeSetupCacheStrategy();
+
 	const double StartTime = FPlatformTime::Seconds();
 
 	IStorageServerHttpClient::FResult ResultTuple;
@@ -268,7 +274,7 @@ TIoStatusOr<FIoBuffer> FStorageServerConnection::ReadChunkRequest(
 	FIoBuffer CacheChunkBuffer;
 	EStorageServerContentType CacheContentType;
 	// TODO is there a way to pass destination directly?
-	if (CacheStrategy.IsValid() && CacheStrategy->ReadChunk(ChunkId, Offset, Size, TOptional<FIoBuffer>(), CacheChunkBuffer, CacheContentType))
+	if (bCacheAvailable && CacheStrategy.IsValid() && CacheStrategy->ReadChunk(ChunkId, Offset, Size, TOptional<FIoBuffer>(), CacheChunkBuffer, CacheContentType))
 	{
 		bWasCached = true;
 		ZenCacheHit(CacheChunkBuffer.GetSize());
@@ -589,6 +595,8 @@ void FStorageServerConnection::SetupCacheStrategy()
 		return;
 	}
 
+	AsyncQueryLatestServerChunkInfo = MakeShared<FAsyncQueryLatestServerChunkInfo>(*this);
+
 	FString DevStoragePath;
 #if defined(STORAGE_SERVER_PLATFORM_CACHE_ROOT_DIR)
 	DevStoragePath = STORAGE_SERVER_PLATFORM_CACHE_ROOT_DIR;
@@ -605,26 +613,80 @@ void FStorageServerConnection::SetupCacheStrategy()
 	TUniquePtr<StorageServer::ICacheJournal> Journal = MakeUnique<StorageServer::FCacheJournalSimple>(*(DevStoragePath / TEXT("ZenCache.db")), FlushEveryNEntries);
 	TUniquePtr<StorageServer::ICacheStorage> Storage = MakeUnique<StorageServer::FCacheStorageBlocking>(*(DevStoragePath / TEXT("ZenCache.storage")), Size);
 	CacheStrategy = MakeUnique<StorageServer::FCacheStrategyLinear>(MoveTemp(Journal), MoveTemp(Storage), AbandonSize, FlushInterval, bShouldInvalidate);
-
-	RefreshCacheEntries();
 }
 #endif // !HAS_STORAGE_SERVER_PLATFORM_CACHE_STRATEGY
 
-void FStorageServerConnection::RefreshCacheEntries()
+bool FStorageServerConnection::FinalizeSetupCacheStrategy()
 {
-	TRACE_CPUPROFILER_EVENT_SCOPE(StorageServerRefreshCacheEntries);
-
 	if (!CacheStrategy.IsValid())
 	{
-		return;
+		return false;
+	}
+
+	// we can read from the cache if the asyncronous initialization has completed
+	if (!AsyncQueryLatestServerChunkInfo.IsValid())
+	{
+		return true;
+	}
+
+	// don't try to read from the cache until we've got the latest chunk hashes from the server 
+	// this will cause files to be read from the server instead of the cache but all early files are
+	// fairly small and it is the only way to guarantee they are up to date
+	if (!AsyncQueryLatestServerChunkInfo->IsFinished())
+	{
+		return false;
+	}
+
+	// the server chunk info will have been updated by the async init thread.
+	// if there is a platform-specific implementation of SetupCacheStrategy that defers cache creation, the expectation is that it will be handled there
+	AsyncQueryLatestServerChunkInfo.Reset();	
+	return true;
+}
+
+FStorageServerConnection::FAsyncQueryLatestServerChunkInfo::FAsyncQueryLatestServerChunkInfo(FStorageServerConnection& InOwner)
+	: Owner(InOwner)
+	, IsCompleted(FPlatformProcess::GetSynchEventFromPool(true))
+{
+	if (FRunnableThread::Create( this, TEXT("StorageServerCacheEntriesInit"), 0, EThreadPriority::TPri_Normal) == nullptr)
+	{
+		IsCompleted->Trigger();
 	}
+}
+
+FStorageServerConnection::FAsyncQueryLatestServerChunkInfo::~FAsyncQueryLatestServerChunkInfo()
+{
+	IsCompleted->Wait();
+	FPlatformProcess::ReturnSynchEventToPool(IsCompleted);
+}
+
+TArray<StorageServer::FLatestServerChunkInfo> FStorageServerConnection::FAsyncQueryLatestServerChunkInfo::GetResults()
+{
+	IsCompleted->Wait();
+	return MoveTemp(Results);
+}
+bool FStorageServerConnection::FAsyncQueryLatestServerChunkInfo::IsFinished() const
+{
+	return IsCompleted->Wait(0);
+}
+
+uint32 FStorageServerConnection::FAsyncQueryLatestServerChunkInfo::Run()
+{
+	TRACE_CPUPROFILER_EVENT_SCOPE(StorageServerRefreshCacheEntries);
 
 	// TODO replace with a new zen server endpoint when it will become available
-	// TODO make it async
-	ChunkInfosRequest([&](FIoChunkId Id, FIoHash RawHash, int64 RawSize)
+	Owner.ChunkInfosRequest([&](FIoChunkId ChunkId, FIoHash RawHash, int64 RawSize)
 	{
-		CacheStrategy->SetChunkInfo(Id, RawHash, RawSize); // using RawHash as a subtitute for ModHash.
+		Results.Emplace(StorageServer::FLatestServerChunkInfo{ChunkId, RawHash, RawSize}); // using RawHash as a subtitute for ModHash.
 	});
+
+	if (Owner.CacheStrategy.IsValid())
+	{
+		Owner.CacheStrategy->ApplyLatestServerChunkInfo(Results);
+		Results.Reset();
+	}
+
+	IsCompleted->Trigger();
+	return 0;
 }
 
 void FStorageServerConnection::BuildReadChunkRequestUrl(FAnsiStringBuilderBase& Builder, const FIoChunkId& ChunkId, const uint64 Offset, const uint64 Size)
diff --git a/Engine/Source/Runtime/StorageServerClient/Private/StorageServerConnection.h b/Engine/Source/Runtime/StorageServerClient/Private/StorageServerConnection.h
--- a/Engine/Source/Runtime/StorageServerClient/Private/StorageServerConnection.h
+++ b/Engine/Source/Runtime/StorageServerClient/Private/StorageServerConnection.h
@@ -11,6 +11,7 @@
 #include "Templates/SharedPointer.h"
 #include "SocketTypes.h"
 #include "HAL/PlatformTime.h"
+#include "HAL/Runnable.h"
 #include "StorageServerHttpClient.h"
 #include "IStorageServerPlatformFile.h"
 #include "Cache/CacheStrategy.h"
@@ -83,7 +84,7 @@ private:
 	TSharedPtr<FInternetAddr> StringToInternetAddr(const FString Address, const int32 Port);
 	bool HandshakeRequest();
 	void SetupCacheStrategy();
-	void RefreshCacheEntries();
+	bool FinalizeSetupCacheStrategy();
 	void BuildReadChunkRequestUrl(FAnsiStringBuilderBase& Builder, const FIoChunkId& ChunkId, const uint64 Offset, const uint64 Size);
 	static TIoStatusOr<FIoBuffer> ReadChunkRequestProcessHttpResult(
 		IStorageServerHttpClient::FResult ResultTuple,
@@ -94,6 +95,25 @@ private:
 	);
 	static uint64 GetCompressedOffset(const FCompressedBuffer& Buffer, uint64 RawOffset);
 	void AddTimingInstance(const double Duration, const uint64 Bytes);
+
+	class FAsyncQueryLatestServerChunkInfo : public FRunnable
+	{
+	public:
+		FAsyncQueryLatestServerChunkInfo(FStorageServerConnection& InOwner);
+		virtual ~FAsyncQueryLatestServerChunkInfo();
+		
+		bool IsFinished() const;
+		TArray<StorageServer::FLatestServerChunkInfo> GetResults();
+		
+	private:
+		virtual uint32 Run() override;
+
+		FStorageServerConnection& Owner;
+		class FEvent* IsCompleted;
+		TArray<StorageServer::FLatestServerChunkInfo> Results;
+
+	};
+	TSharedPtr<FAsyncQueryLatestServerChunkInfo, ESPMode::ThreadSafe> AsyncQueryLatestServerChunkInfo;
 };
 
 #endif
