Mmap backing storage for zen client side cache

diff --git a/Engine/Source/Runtime/StorageServerClient/Private/Cache/CacheStorage.h b/Engine/Source/Runtime/StorageServerClient/Private/Cache/CacheStorage.h
--- a/Engine/Source/Runtime/StorageServerClient/Private/Cache/CacheStorage.h
+++ b/Engine/Source/Runtime/StorageServerClient/Private/Cache/CacheStorage.h
@@ -1,9 +1,13 @@
-﻿// Copyright Epic Games, Inc. All Rights Reserved.
+// Copyright Epic Games, Inc. All Rights Reserved.
 
 #pragma once
 
 #include "IO/IoChunkId.h"
 #include "IO/IoBuffer.h"
+#include "Containers/Array.h"
+#include "Templates/Tuple.h"
+#include "Containers/UnrealString.h"
+#include "Math/UnrealMathUtility.h"
 
 #if !UE_BUILD_SHIPPING
 
@@ -44,6 +48,57 @@ namespace StorageServer
 
 		// Limit cache file size to 2GB to stay below file system limitations
 		static constexpr uint64 MaxCacheFileSize = 2ull * 1024 * 1024 * 1024;
+
+		// Generate a list of filenames and sizes to be used for backing storage
+		static TArray<TTuple<FString, uint64>> GetBackingFileNames(const TCHAR* FileNamePrefix, const uint64 FileSizeTmp)
+		{
+			TArray<TTuple<FString, uint64>> BackingFiles;
+			BackingFiles.SetNum(FMath::DivideAndRoundUp(FileSizeTmp, MaxCacheFileSize));
+
+			const uint64 SizeOfLastFile = FileSizeTmp - MaxCacheFileSize * (BackingFiles.Num() - 1);
+
+			for (int32 i = 0; i < BackingFiles.Num(); ++i)
+			{
+				// TODO check if disk size is exhausted?
+				const FString FileName = FString::Printf(TEXT("%s%u"), FileNamePrefix, i);
+				const uint64 DesiredFileSize = i == BackingFiles.Num() - 1 ? SizeOfLastFile : MaxCacheFileSize;
+
+				BackingFiles[i] = { FileName, DesiredFileSize };
+			}
+
+			return BackingFiles;
+		}
+
+		static bool GetBackingIntervals(const uint64 Offset, const uint64 Size,
+										 uint32& IndexA, uint64& OffsetA, uint64& SizeA,
+										 uint32& IndexB, uint64& OffsetB, uint64& SizeB)
+		{
+			IndexA = FMath::DivideAndRoundDown(Offset, MaxCacheFileSize);
+			IndexB = FMath::DivideAndRoundDown(Offset + Size - 1, MaxCacheFileSize);
+
+			// we only support writes spanning maximum two backing storage files
+			if (!ensureAlways(IndexA + 1 >= IndexB))
+			{
+				return false;
+			}
+
+			if (IndexA == IndexB)
+			{
+				OffsetA = Offset - IndexA * MaxCacheFileSize;
+				SizeA = Size;
+				OffsetB = 0;
+				SizeB = 0;
+				return true;
+			}
+			else
+			{
+				OffsetA = Offset - IndexA * MaxCacheFileSize;
+				SizeA = IndexB * MaxCacheFileSize - Offset;
+				OffsetB = 0;
+				SizeB = Size - SizeA;
+				return true;
+			}
+		}
 	};
 }
 
diff --git a/Engine/Source/Runtime/StorageServerClient/Private/Cache/CacheStorageBlocking.cpp b/Engine/Source/Runtime/StorageServerClient/Private/Cache/CacheStorageBlocking.cpp
--- a/Engine/Source/Runtime/StorageServerClient/Private/Cache/CacheStorageBlocking.cpp
+++ b/Engine/Source/Runtime/StorageServerClient/Private/Cache/CacheStorageBlocking.cpp
@@ -17,24 +17,23 @@ namespace StorageServer
 	{
 		IPlatformFile& PlatformFile = FPlatformFileManager::Get().GetPlatformFile();
 
-		BackingFiles.SetNum(FMath::DivideAndRoundUp(FileSizeTmp, MaxCacheFileSize));
-		FileSize = 0;
+		const TArray<TTuple<FString, uint64>> BackingFileNames = GetBackingFileNames(FileNamePrefix, FileSizeTmp);
 
-		const uint64 SizeOfLastFile = FileSizeTmp - MaxCacheFileSize * (BackingFiles.Num() - 1);
+		BackingFiles.SetNum(BackingFileNames.Num());
+		FileSize = 0;
 
 		for (int32 i = 0; i < BackingFiles.Num(); ++i)
 		{
+			const FString& FileName = BackingFileNames[i].Key;
+			const uint64 DesiredFileSize = BackingFileNames[i].Value;
+
 			// TODO check if disk size is exhausted?
-			// TODO use temp location that is available on development platforms
-			FString FileName = FString::Printf(TEXT("%s%u"), FileNamePrefix, i);
 			IFileHandle* BackingFile = PlatformFile.OpenWrite(*FileName, true, true);
 			if (!ensureAlwaysMsgf(BackingFile, TEXT("Can't open storage server backing file '%s'"), *FileName))
 			{
 				return;
 			}
 
-			const uint64 DesiredFileSize = i == BackingFiles.Num() - 1 ? SizeOfLastFile : MaxCacheFileSize;
-
 			if (BackingFile->Size() != DesiredFileSize)
 			{
 				BackingFile->Truncate(DesiredFileSize);
@@ -153,37 +152,14 @@ namespace StorageServer
 		}, LowLevelTasks::ETaskPriority::BackgroundNormal);
 	}
 
-	void FCacheStorageBlocking::BackingIntervals(const uint64 Offset, const uint64 Size,
-												 uint32& IndexA, uint64& OffsetA, uint64& SizeA,
-												 uint32& IndexB, uint64& OffsetB, uint64& SizeB) const
-	{
-		IndexA = FMath::DivideAndRoundDown(Offset, MaxCacheFileSize);
-		IndexB = FMath::DivideAndRoundDown(Offset + Size - 1, MaxCacheFileSize);
-
-		// we only support writes spanning maximum two backing storage files
-		ensureAlways(IndexA + 1 >= IndexB);
-
-		if (IndexA == IndexB)
-		{
-			OffsetA = Offset - IndexA * MaxCacheFileSize;
-			SizeA = Size;
-			OffsetB = 0;
-			SizeB = 0;
-		}
-		else
-		{
-			OffsetA = Offset - IndexA * MaxCacheFileSize;
-			SizeA = IndexB * MaxCacheFileSize - Offset;
-			OffsetB = 0;
-			SizeB = Size - SizeA;
-		}
-	}
-
 	void FCacheStorageBlocking::BackingReadAt(const uint64 Offset, void* Buffer, const uint64 Size)
 	{
 		uint32 IndexA, IndexB;
 		uint64 OffsetA, SizeA, OffsetB, SizeB;
-		BackingIntervals(Offset, Size, IndexA, OffsetA, SizeA, IndexB, OffsetB, SizeB);
+		if (!GetBackingIntervals(Offset, Size, IndexA, OffsetA, SizeA, IndexB, OffsetB, SizeB))
+		{
+			return;
+		}
 
 		if (SizeA > 0)
 		{
@@ -200,7 +176,10 @@ namespace StorageServer
 	{
 		uint32 IndexA, IndexB;
 		uint64 OffsetA, SizeA, OffsetB, SizeB;
-		BackingIntervals(Offset, Size, IndexA, OffsetA, SizeA, IndexB, OffsetB, SizeB);
+		if (!GetBackingIntervals(Offset, Size, IndexA, OffsetA, SizeA, IndexB, OffsetB, SizeB))
+		{
+			return;
+		}
 
 		FScopeLock ScopeFileWriteLock(&FileWriteLock);
 
diff --git a/Engine/Source/Runtime/StorageServerClient/Private/Cache/CacheStorageBlocking.h b/Engine/Source/Runtime/StorageServerClient/Private/Cache/CacheStorageBlocking.h
--- a/Engine/Source/Runtime/StorageServerClient/Private/Cache/CacheStorageBlocking.h
+++ b/Engine/Source/Runtime/StorageServerClient/Private/Cache/CacheStorageBlocking.h
@@ -101,9 +101,6 @@ namespace StorageServer
 		uint64 FileSize;
 
 		void Flush(bool bLazy);
-		void BackingIntervals(const uint64 Offset, const uint64 Size,
-			uint32& IndexA, uint64& OffsetA, uint64& SizeA,
-			uint32& IndexB, uint64& OffsetB, uint64& SizeB) const;
 		void BackingReadAt(const uint64 Offset, void* Buffer, const uint64 Size);
 		void BackingWriteAt(const uint64 Offset, const void* Buffer, const uint64 Size);
 	};
diff --git a/Engine/Source/Runtime/StorageServerClient/Private/Cache/CacheStorageMmap.cpp b/Engine/Source/Runtime/StorageServerClient/Private/Cache/CacheStorageMmap.cpp
new file mode 100644
--- /dev/null
+++ b/Engine/Source/Runtime/StorageServerClient/Private/Cache/CacheStorageMmap.cpp
@@ -0,0 +1,145 @@
+﻿// Copyright Epic Games, Inc. All Rights Reserved.
+
+#include "CacheStorageMmap.h"
+
+#if !UE_BUILD_SHIPPING
+
+#include "HAL/PlatformFileManager.h"
+#include "GenericPlatform/GenericPlatformFile.h"
+
+namespace StorageServer
+{
+	FCacheStorageMmap::FCacheStorageMmap(const TCHAR* FileNamePrefix, const uint64 FileSizeTmp)
+	{
+		if (!ensure(FPlatformProperties::SupportsMemoryMappedFiles()))
+		{
+			return;
+		}
+
+		const TArray<TTuple<FString, uint64>> BackingFileNames = GetBackingFileNames(FileNamePrefix, FileSizeTmp);
+
+		BackingFiles.SetNum(BackingFileNames.Num());
+		TotalSize = 0;
+
+		IPlatformFile& PlatformFile = FPlatformFileManager::Get().GetPlatformFile();
+
+		for (int32 i = 0; i < BackingFiles.Num(); ++i)
+		{
+			const FString& FileName = BackingFileNames[i].Key;
+			const uint64 DesiredFileSize = BackingFileNames[i].Value;
+			FBackingFile& BackingFile = BackingFiles[i];
+
+			IFileHandle* PlainFileHandle = PlatformFile.OpenWrite(*FileName, true, true);
+			if (!ensureAlwaysMsgf(PlainFileHandle, TEXT("FCacheStorageMmap failed to open cache storage file '%s'"), *FileName))
+			{
+				return;
+			}
+
+			if (PlainFileHandle->Size() != DesiredFileSize)
+			{
+				// TODO add truncate to IMappedFileHandle?
+				PlainFileHandle->Truncate(DesiredFileSize);
+				bNewlyCreatedStorage = true;
+			}
+
+			const uint64 ActualFileSize = PlainFileHandle->Size();
+			delete PlainFileHandle;
+
+			FOpenMappedResult Result = PlatformFile.OpenMappedEx(*FileName, IPlatformFile::EOpenReadFlags::AllowWrite);
+			if (!ensureAlwaysMsgf(!Result.HasError(), TEXT("FCacheStorageMmap OpenMappedEx failed to open '%s' due to '%s'"), *FileName, *Result.GetError().GetMessage()))
+			{
+				return;
+			}
+
+			BackingFile.FileHandle = Result.StealValue();
+			BackingFile.FileRegion = TUniquePtr<IMappedFileRegion>(BackingFile.FileHandle->MapRegion(0, ActualFileSize, EMappedFileFlags::EFileWritable));
+			BackingFile.MapPtr = (uint8*)BackingFile.FileRegion->GetMappedPtr(); // TODO add rw ptr to IMappedFileRegion
+			BackingFile.MapSize = BackingFile.FileRegion->GetMappedSize();
+
+			TotalSize += BackingFile.MapSize;
+		}
+	}
+
+	FCacheStorageMmap::~FCacheStorageMmap()
+	{
+		BackingFiles.Reset();
+		TotalSize = 0;
+	}
+
+	void FCacheStorageMmap::Flush()
+	{
+	}
+
+	void FCacheStorageMmap::Invalidate(const uint64 Offset, const uint64 Size)
+	{
+	}
+
+	FIoBuffer FCacheStorageMmap::Read(const uint64 Offset, const uint64 ReadSize, TOptional<FIoBuffer> OptDestination)
+	{
+		if (!IsValidRange(Offset, ReadSize))
+		{
+			return FIoBuffer();
+		}
+
+		uint32 IndexA, IndexB;
+		uint64 OffsetA, SizeA, OffsetB, SizeB;
+		if (!GetBackingIntervals(Offset, ReadSize, IndexA, OffsetA, SizeA, IndexB, OffsetB, SizeB))
+		{
+			return FIoBuffer();
+		}
+
+		if (SizeA > 0 && SizeB == 0)
+		{
+			if (OptDestination.IsSet() && OptDestination->GetSize() >= SizeA)
+			{
+				OptDestination->SetSize(SizeA);
+				FMemory::Memcpy(OptDestination->GetData(), BackingFiles[IndexA].MapPtr + OffsetA, SizeA);
+				return OptDestination.GetValue();
+			}
+			else
+			{
+				return FIoBuffer(FIoBuffer::Wrap, BackingFiles[IndexA].MapPtr + OffsetA, SizeA);
+			}
+		}
+
+		if (SizeA > 0 && SizeB > 0)
+		{
+			FIoBuffer Result = (OptDestination.IsSet() && OptDestination->GetSize() >= SizeA + SizeB) ? OptDestination.GetValue() : FIoBuffer();
+			Result.SetSize(SizeA + SizeB);
+			FMemory::Memcpy(Result.GetData(), BackingFiles[IndexA].MapPtr + OffsetA, SizeA);
+			FMemory::Memcpy(Result.GetData() + SizeA, BackingFiles[IndexB].MapPtr + OffsetB, SizeB);
+			return Result;
+		}
+
+		return FIoBuffer();
+	}
+
+	void FCacheStorageMmap::WriteAsync(const uint64 Offset, const void* Buffer, const uint64 WriteSize)
+	{
+		if (!IsValidRange(Offset, WriteSize))
+		{
+			return;
+		}
+
+		uint32 IndexA, IndexB;
+		uint64 OffsetA, SizeA, OffsetB, SizeB;
+		if (!GetBackingIntervals(Offset, WriteSize, IndexA, OffsetA, SizeA, IndexB, OffsetB, SizeB))
+		{
+			return;
+		}
+
+		if (SizeA > 0 && SizeB == 0)
+		{
+			FMemory::Memcpy(BackingFiles[IndexA].MapPtr + OffsetA, Buffer, SizeA);
+		}
+
+		if (SizeA > 0 && SizeB > 0)
+		{
+			FMemory::Memcpy(BackingFiles[IndexA].MapPtr + OffsetA, Buffer, SizeA);
+			FMemory::Memcpy(BackingFiles[IndexB].MapPtr + OffsetB, (const uint8*)Buffer + SizeA, SizeB);
+		}
+	}
+}
+
+#endif
+
diff --git a/Engine/Source/Runtime/StorageServerClient/Private/Cache/CacheStorageMmap.h b/Engine/Source/Runtime/StorageServerClient/Private/Cache/CacheStorageMmap.h
new file mode 100644
--- /dev/null
+++ b/Engine/Source/Runtime/StorageServerClient/Private/Cache/CacheStorageMmap.h
@@ -0,0 +1,70 @@
+﻿// Copyright Epic Games, Inc. All Rights Reserved.
+
+#pragma once
+
+#include "CacheStorage.h"
+
+#if !UE_BUILD_SHIPPING
+
+#include "Async/MappedFileHandle.h"
+
+namespace StorageServer
+{
+	// Cache storage implemented via mmap-ing a file on platforms that support mmap
+	class FCacheStorageMmap : public ICacheStorage
+	{
+	public:
+		FCacheStorageMmap(const TCHAR* FileNamePrefix, const uint64 FileSize);
+		virtual ~FCacheStorageMmap() override;
+
+		virtual void Flush() override;
+
+		virtual uint64 GetSize() const override {return TotalSize;}
+
+		virtual void Invalidate(const uint64 Offset, const uint64 Size) override;
+
+		virtual FIoBuffer Read(
+			const uint64 Offset,
+			const uint64 ReadSize,
+			TOptional<FIoBuffer> OptDestination
+		) override;
+
+		virtual void WriteAsync(
+			const uint64 Offset,
+			const void* Buffer,
+			const uint64 WriteSize
+		) override;
+
+	private:
+		struct FBackingFile
+		{
+			TUniquePtr<IMappedFileHandle> FileHandle;
+			TUniquePtr<IMappedFileRegion> FileRegion;
+			uint8* MapPtr;
+			uint64 MapSize;
+
+			FBackingFile() = default;
+			FBackingFile(FBackingFile&&) = default;
+			FBackingFile(const FBackingFile&) = delete;
+			FBackingFile& operator=(const FBackingFile&) = delete;
+			~FBackingFile()
+			{
+				// release region before file handle
+				FileRegion.Reset();
+				FileHandle.Reset();
+				MapPtr = nullptr;
+				MapSize = 0;
+			}
+		};
+
+		TArray<FBackingFile> BackingFiles;
+		uint64 TotalSize;
+
+		bool IsValidRange(const uint64 Offset, const uint64 Size) const
+		{
+			return Offset + Size <= TotalSize;
+		}
+	};
+}
+
+#endif
\ No newline at end of file
