SSC: Fixing crashing on shutdown with the eventpool.

When we are exiting, ~TEventPool has already run, so we end up in a bad scenario which leads to us failing to get the Pool to return anything to. Since we're shutting down, everything is gone already.

diff --git a/Engine/Source/Runtime/StorageServerClient/Private/StorageServerPlatformFile.cpp b/Engine/Source/Runtime/StorageServerClient/Private/StorageServerPlatformFile.cpp
--- a/Engine/Source/Runtime/StorageServerClient/Private/StorageServerPlatformFile.cpp
+++ b/Engine/Source/Runtime/StorageServerClient/Private/StorageServerPlatformFile.cpp
@@ -249,14 +249,6 @@ public:
 
 	virtual ~FStorageServerEngineStartupPrecache()
 	{
-		// the worker thread has most likely completed already but no harm in shutting down cleanly
-		bExitImmediately = true;
-		HasWork->Trigger();
-		IsCompleted->Wait();
-
-		FPlatformProcess::ReturnSynchEventToPool(HasWork);
-		FPlatformProcess::ReturnSynchEventToPool(IsCompleted);
-
 #if STORAGE_SERVER_START_CACHE_REPORT
 		UE_CLOG(Cache.Num() > 0, LogStorageServerPlatformFile, Log, TEXT("Engine startup precache size was %llu bytes, %d items. %d items actually retrieved (%d%%). %d cache misses (%d%%)"), CacheSize.load(), Cache.Num(), AccessedChunks.Num(), (AccessedChunks.Num()*100) / Cache.Num(), UncachedChunks.Num(), (UncachedChunks.Num()*100) / (AccessedChunks.Num()+UncachedChunks.Num()) );
 #endif
