Update precache of .ini and .uplugin files to request chunks in batch from server.

[RN] minor

diff --git a/Engine/Source/Runtime/StorageServerClient/Private/BuiltInHttpClient/BuiltInHttpClient.cpp b/Engine/Source/Runtime/StorageServerClient/Private/BuiltInHttpClient/BuiltInHttpClient.cpp
--- a/Engine/Source/Runtime/StorageServerClient/Private/BuiltInHttpClient/BuiltInHttpClient.cpp
+++ b/Engine/Source/Runtime/StorageServerClient/Private/BuiltInHttpClient/BuiltInHttpClient.cpp
@@ -10,7 +10,7 @@ static inline FAnsiStringView GetMimeTypeString(EStorageServerContentType Conten
 	switch (ContentType)
 	{
 		case EStorageServerContentType::CbObject:
-			return ANSITEXTVIEW("application/ue-x-cb");
+			return ANSITEXTVIEW("application/x-ue-cb");
 		case EStorageServerContentType::Binary:
 			return ANSITEXTVIEW("application/octet-stream");
 		case EStorageServerContentType::CompressedBinary:
@@ -51,6 +51,7 @@ IStorageServerHttpClient::FResult FBuiltInHttpClient::RequestSync(
 	EStorageServerContentType Accept,
 	FAnsiStringView Verb,
 	TOptional<FIoBuffer> OptPayload,
+	EStorageServerContentType PayloadContentType,
 	TOptional<FIoBuffer> OptDestination,
 	float TimeoutSeconds,
 	const bool bReportErrors
@@ -74,6 +75,11 @@ IStorageServerHttpClient::FResult FBuiltInHttpClient::RequestSync(
 		HeaderBuffer
 			<< "Content-Length: " << OptPayload->GetSize() << "\r\n";
 	}
+	if (PayloadContentType != EStorageServerContentType::Unknown)
+	{
+		HeaderBuffer
+			<< "Content-Type: " << GetMimeTypeString(PayloadContentType) << "\r\n";
+	}
 	HeaderBuffer << "\r\n";
 
 	IBuiltInHttpClientSocket* Socket = nullptr;
diff --git a/Engine/Source/Runtime/StorageServerClient/Private/BuiltInHttpClient/BuiltInHttpClient.h b/Engine/Source/Runtime/StorageServerClient/Private/BuiltInHttpClient/BuiltInHttpClient.h
--- a/Engine/Source/Runtime/StorageServerClient/Private/BuiltInHttpClient/BuiltInHttpClient.h
+++ b/Engine/Source/Runtime/StorageServerClient/Private/BuiltInHttpClient/BuiltInHttpClient.h
@@ -38,6 +38,7 @@ public:
 		EStorageServerContentType Accept,
 		FAnsiStringView Verb,
 		TOptional<FIoBuffer> OptPayload,
+		EStorageServerContentType PayloadContentType,
 		TOptional<FIoBuffer> OptDestination,
 		float TimeoutSeconds,
 		const bool bReportErrors
@@ -49,13 +50,14 @@ public:
 		EStorageServerContentType Accept,
 		FAnsiStringView Verb,
 		TOptional<FIoBuffer> OptPayload,
+		EStorageServerContentType PayloadContentType,
 		TOptional<FIoBuffer> OptDestination,
 		float TimeoutSeconds,
 		const bool bReportErrors
 	) override
 	{
 		// This HTTP client doesn't support async requests
-		Callback(RequestSync(Url, Accept, Verb, OptPayload, OptDestination, TimeoutSeconds, bReportErrors));
+		Callback(RequestSync(Url, Accept, Verb, OptPayload, PayloadContentType, OptDestination, TimeoutSeconds, bReportErrors));
 	}
 private:
 	TUniquePtr<IBuiltInHttpClientSocketPool> SocketPool;
diff --git a/Engine/Source/Runtime/StorageServerClient/Private/StorageServerConnection.cpp b/Engine/Source/Runtime/StorageServerClient/Private/StorageServerConnection.cpp
--- a/Engine/Source/Runtime/StorageServerClient/Private/StorageServerConnection.cpp
+++ b/Engine/Source/Runtime/StorageServerClient/Private/StorageServerConnection.cpp
@@ -12,6 +12,7 @@
 #include "Serialization/MemoryReader.h"
 #include "Serialization/CompactBinary.h"
 #include "Serialization/CompactBinarySerialization.h"
+#include "Serialization/CompactBinaryWriter.h"
 #include "Memory/MemoryView.h"
 #include "Memory/SharedBuffer.h"
 #include "ProfilingDebugging/CountersTrace.h"
@@ -307,6 +308,61 @@ TIoStatusOr<FIoBuffer> FStorageServerConnection::ReadChunkRequest(
 	return ResultBuffer;
 }
 
+TIoStatusOr<FIoBuffer> FStorageServerConnection::ReadChunkBatchRequest(const TArray<FChunkBatchRequestEntry>& Chunks)
+{
+	TRACE_CPUPROFILER_EVENT_SCOPE(FStorageServerConnection::ReadChunkBatchRequest);
+
+	const double StartTime = FPlatformTime::Seconds();
+
+	FCbWriter Writer;
+
+	Writer.BeginObject();
+
+	Writer << "method" << "getchunks";
+
+	Writer.BeginObject("Request");
+	Writer.BeginArray("Chunks");
+
+	for (const FChunkBatchRequestEntry& Chunk : Chunks)
+	{
+		Writer.BeginObject();
+
+		Writer << "Offset" << Chunk.Offset;
+		Writer << "Size" << Chunk.Size;
+		Writer << "Oid" << Chunk.Id;
+
+		Writer.EndObject();
+	}
+
+	Writer.EndArray();
+	Writer.EndObject();
+	Writer.EndObject();
+
+	TAnsiStringBuilder<256> Uri;
+	Uri.Append(BaseURI).Append("/rpc");
+
+	FIoBuffer Payload(Writer.GetSaveSize());
+	Writer.Save(Payload.GetMutableView());
+
+	IStorageServerHttpClient::FResult HttpResult = HttpClient->RequestSync(
+		*Uri,
+		EStorageServerContentType::CompressedBinary,
+		"POST",
+		Payload,
+		EStorageServerContentType::CbObject
+	);
+
+	TIoStatusOr<FIoBuffer> Result = HttpResult.Get<0>();
+
+	if (Result.IsOk())
+	{
+		const double Duration = FPlatformTime::Seconds() - StartTime;
+		AddTimingInstance(Duration, Result.ValueOrDie().GetSize());
+	}
+
+	return Result;
+}
+
 void FStorageServerConnection::ReadChunkRequestAsync(
 	const FIoChunkId& ChunkId,
 	const uint64 Offset,
@@ -567,6 +623,7 @@ bool FStorageServerConnection::HandshakeRequest()
 		EStorageServerContentType::Unknown,
 		"GET",
 		TOptional<FIoBuffer>(),
+		EStorageServerContentType::Unknown,
 		TOptional<FIoBuffer>(),
 		ConnectionTimeoutSeconds,
 		false
diff --git a/Engine/Source/Runtime/StorageServerClient/Private/StorageServerConnection.h b/Engine/Source/Runtime/StorageServerClient/Private/StorageServerConnection.h
--- a/Engine/Source/Runtime/StorageServerClient/Private/StorageServerConnection.h
+++ b/Engine/Source/Runtime/StorageServerClient/Private/StorageServerConnection.h
@@ -6,6 +6,7 @@
 #include "Compression/CompressedBuffer.h"
 #include "Containers/UnrealString.h"
 #include "Containers/AnsiString.h"
+#include "IO/IoChunkId.h"
 #include "Misc/StringBuilder.h"
 #include "Memory/MemoryFwd.h"
 #include "Templates/SharedPointer.h"
@@ -29,7 +30,6 @@
 DECLARE_LOG_CATEGORY_EXTERN(LogStorageServerConnection, Log, All);
 
 struct FPackageStoreEntryResource;
-class FIoChunkId;
 
 class FStorageServerConnection
 {
@@ -50,6 +50,16 @@ public:
 		const TOptional<FIoBuffer> OptDestination,
 		const bool bHardwareTargetBuffer
 	);
+
+	// Matches RequestChunkEntry in zenserver httprojectstore.cpp
+	struct FChunkBatchRequestEntry {
+		FIoChunkId Id;
+		uint64 Offset;
+		uint64 Size;
+	};
+
+	TIoStatusOr<FIoBuffer> ReadChunkBatchRequest(const TArray<FChunkBatchRequestEntry>& Chunks);
+	
 	void ReadChunkRequestAsync(
 		const FIoChunkId& ChunkId,
 		const uint64 Offset,
diff --git a/Engine/Source/Runtime/StorageServerClient/Private/StorageServerHttpClient.h b/Engine/Source/Runtime/StorageServerClient/Private/StorageServerHttpClient.h
--- a/Engine/Source/Runtime/StorageServerClient/Private/StorageServerHttpClient.h
+++ b/Engine/Source/Runtime/StorageServerClient/Private/StorageServerHttpClient.h
@@ -30,6 +30,7 @@ public:
 		EStorageServerContentType Accept = EStorageServerContentType::Unknown,
 		FAnsiStringView Verb = "GET",
 		TOptional<FIoBuffer> OptPayload = TOptional<FIoBuffer>(),
+		EStorageServerContentType PayloadContentType = EStorageServerContentType::Unknown,
 		TOptional<FIoBuffer> OptDestination = TOptional<FIoBuffer>(),
 		float TimeoutSeconds = -1.f,
 		const bool bReportErrors = true
@@ -41,6 +42,7 @@ public:
 		EStorageServerContentType Accept = EStorageServerContentType::Unknown,
 		FAnsiStringView Verb = "GET",
 		TOptional<FIoBuffer> OptPayload = TOptional<FIoBuffer>(),
+		EStorageServerContentType PayloadContentType = EStorageServerContentType::Unknown,
 		TOptional<FIoBuffer> OptDestination = TOptional<FIoBuffer>(),
 		float TimeoutSeconds = -1.f,
 		const bool bReportErrors = true
diff --git a/Engine/Source/Runtime/StorageServerClient/Private/StorageServerPlatformFile.cpp b/Engine/Source/Runtime/StorageServerClient/Private/StorageServerPlatformFile.cpp
--- a/Engine/Source/Runtime/StorageServerClient/Private/StorageServerPlatformFile.cpp
+++ b/Engine/Source/Runtime/StorageServerClient/Private/StorageServerPlatformFile.cpp
@@ -9,6 +9,7 @@
 #include "HAL/RunnableThread.h"
 #include "HAL/CriticalSection.h"
 #include "Misc/App.h"
+#include "Misc/AssertionMacros.h"
 #include "Misc/CommandLine.h"
 #include "Misc/CoreDelegates.h"
 #include "Misc/FileHelper.h"
@@ -18,7 +19,9 @@
 #include "Misc/StringBuilder.h"
 #include "Modules/ModuleManager.h"
 #include "Modules/ModuleManager.h"
+#include "Serialization/CompactBinaryPackage.h"
 #include "Serialization/CompactBinarySerialization.h"
+#include "Serialization/CompactBinaryWriter.h"
 #include "Serialization/JsonReader.h"
 #include "Serialization/JsonSerializer.h"
 #include "ProfilingDebugging/PlatformFileTrace.h"
@@ -322,6 +325,10 @@ private:
 			PrecacheItems();
 			if (bExitWhenComplete)
 			{
+				// Clear out the final batch of requests which may have been added after the previous
+				// PrecacheItems() call had processed the queue, but before the HTTP response had
+				// returned from the server.
+				PrecacheItems();
 				break;
 			}
 		}
@@ -330,35 +337,123 @@ private:
 		return 0;
 	}
 
+	void HandleChunkBatchResponse(const FIoBuffer& Payload)
+	{
+		struct ResponseHeader
+		{
+			uint32 Magic;
+			uint32 AttachmentCount;
+			uint32 Reserved[2];
+		} Header;
+
+		FMemoryReaderView View(Payload.GetView());
+		View.Serialize(&Header, sizeof(Header));
+
+		struct AttachmentInfo
+		{
+			uint64 PayloadSize;
+			uint32 Flags;
+			FIoHash Hash;
+		};
+
+		if (Header.Magic != 0xaa77aacc)
+		{
+			UE_LOG(LogStorageServerPlatformFile, Warning, TEXT("Read incorrect header magic from server response - expected %x, got %x"), 0xaa77aacc, Header.Magic);
+			return;
+		}
+
+		TArray<AttachmentInfo> Info;
+		Info.AddUninitialized(Header.AttachmentCount + 1); // one extra for the CbPackage root object
+
+		for (int32 i = 0; i < Info.Num(); ++i)
+		{
+			View.Serialize(&Info[i], sizeof(Info[i]));
+		}
+
+		FCbObjectView Root(Payload.GetData() + View.Tell());
+		FCbArrayView Chunks = Root["Chunks"].AsArrayView();
+
+		TMap<FIoHash, TArray<uint8>> AttachmentData;
+
+		uint64 Offset = View.Tell() + Root.GetSize();
+		for (int32 i = 1; i < Info.Num(); ++i) // skip the attachment info for the root object
+		{
+			FCompressedBuffer Buffer = FCompressedBuffer::FromCompressed(FSharedBuffer::MakeView(Payload.GetData() + Offset, Info[i].PayloadSize));
+			FCompressedBufferReader Reader(Buffer);
+
+			TArray<uint8> Data;
+			Data.AddUninitialized(Reader.GetRawSize());
+
+			if (!ensureMsgf(Reader.TryDecompressTo(FMutableMemoryView(Data.GetData(), Data.Num())), TEXT("Failed to decompress data from server response")))
+			{
+				continue;
+			}
+
+			AttachmentData.Add(Info[i].Hash, MoveTemp(Data));
+			Offset += Info[i].PayloadSize;
+		}
+
+		for (FCbFieldView Field : Chunks)
+		{
+			FCbObjectView Chunk = Field.AsObjectView();
+			FCbObjectId Id = Chunk["Id"].AsObjectId();
+			FIoHash RawHash = Chunk["RawHash"].AsHash();
+			FIoHash FragmentHash = Chunk["FragmentHash"].AsHash();
+
+			if (!ensureMsgf(RawHash != FIoHash() || FragmentHash != FIoHash(), TEXT("Failed to find hash in chunk info returned from server")))
+			{
+				continue;
+			}
+
+			FIoChunkId ChunkId;
+			ChunkId.Set(Id.GetView());
+			TArray<uint8>* Data = AttachmentData.Find(RawHash == FIoHash() ? FragmentHash : RawHash);
+
+			{
+				FWriteScopeLock Lock(CacheLock);
+				Cache.Add(ChunkId, *Data);
+			}
+		}
+	}
+
 	void PrecacheItems()
 	{
+		TRACE_CPUPROFILER_EVENT_SCOPE(FStorageServerEngineStartupPrecache::PrecacheItems);
+
 		TTuple<FIoChunkId,int32> Item;
+		uint32 CorrelationId = 0;
+
+		TArray<FStorageServerConnection::FChunkBatchRequestEntry> Items;
+
 		while (HighPriorityPrecacheFileRequests.Dequeue(Item) || PrecacheFileRequests.Dequeue(Item))
 		{
-			const FIoChunkId& FileChunkId = Item.Key;
-			const int32& FileSize = Item.Value;
+			const FIoChunkId& Id = Item.Key;
+			const int32& Size = Item.Value;
 
 			// make sure the total cache size doesn't grow out of control
-			if (CacheSize+FileSize > MaxCacheSize)
+			if (CacheSize + Size > MaxCacheSize)
 			{
-				continue;
+				break;
 			}
 
-			TArray<uint8> TempBuffer;
-			TempBuffer.AddUninitialized(FileSize);
+			CacheSize += Size;
+			Items.Add(FStorageServerConnection::FChunkBatchRequestEntry{Id, 0, uint64(Size)});
+		}
 
-			TRACE_CPUPROFILER_EVENT_SCOPE(FStorageServerEngineStartupPrecache::SendReadMessage);
-			TIoStatusOr<FIoBuffer> Result = Connection.ReadChunkRequest(FileChunkId, 0, FileSize, FIoBuffer(FIoBuffer::Wrap, TempBuffer.GetData(), FileSize), false);
-			if (Result.IsOk())
-			{
-				TempBuffer.SetNum(Result.ValueOrDie().GetSize());
-				CacheSize += TempBuffer.Num();
-				{
-					FWriteScopeLock Lock(CacheLock);
-					Cache.Add(FileChunkId, MoveTemp(TempBuffer));
-				}
-			}
+		if (Items.IsEmpty())
+		{
+			return;
 		}
+
+		TIoStatusOr<FIoBuffer> Response = Connection.ReadChunkBatchRequest(Items);	
+
+		if (!ensureMsgf(Response.IsOk(), TEXT("Failed to read chunk batch request from Zen server")))
+		{
+			return;
+		}
+
+		FIoBuffer ResponseBuffer = Response.ValueOrDie();
+		HandleChunkBatchResponse(ResponseBuffer);
 	}
 
 	FStorageServerConnection& Connection;
