Refactor StorageServerClient to use RPC getchunks endpoint for all read chunk requests. Treat missing chunks as warnings instead of errors.

diff --git a/Engine/Source/Runtime/StorageServerClient/Private/BuiltInHttpClient/BuiltInHttpClientFSocket.cpp b/Engine/Source/Runtime/StorageServerClient/Private/BuiltInHttpClient/BuiltInHttpClientFSocket.cpp
--- a/Engine/Source/Runtime/StorageServerClient/Private/BuiltInHttpClient/BuiltInHttpClientFSocket.cpp
+++ b/Engine/Source/Runtime/StorageServerClient/Private/BuiltInHttpClient/BuiltInHttpClientFSocket.cpp
@@ -101,6 +101,8 @@ IBuiltInHttpClientSocket* FBuiltInHttpClientFSocketPool::AcquireSocket(float Tim
 
 	FSocket* Socket = SocketSubsystem.CreateSocket(NAME_Stream, TEXT("StorageServer"), ServerAddr->GetProtocolType());
 
+	Socket->SetNoDelay(true);
+
 	if (TimeoutSeconds > 0.0f)
 	{
 		Socket->SetNonBlocking(true);
diff --git a/Engine/Source/Runtime/StorageServerClient/Private/StorageServerConnection.cpp b/Engine/Source/Runtime/StorageServerClient/Private/StorageServerConnection.cpp
--- a/Engine/Source/Runtime/StorageServerClient/Private/StorageServerConnection.cpp
+++ b/Engine/Source/Runtime/StorageServerClient/Private/StorageServerConnection.cpp
@@ -263,6 +263,7 @@ TIoStatusOr<FIoBuffer> FStorageServerConnection::ReadChunkRequest(
 	const bool bHardwareTargetBuffer
 )
 {
+	// TODO move caching functionality to ReadChunkBatchRequest and remove ReadChunkRequest. 
 	TRACE_CPUPROFILER_EVENT_SCOPE(FStorageServerConnection::ReadChunkRequest);
 
 	bool bCacheAvailable = FinalizeSetupCacheStrategy();
@@ -283,12 +284,29 @@ TIoStatusOr<FIoBuffer> FStorageServerConnection::ReadChunkRequest(
 	}
 	else
 	{
+#if HAS_STORAGE_SERVER_RPC_GETCHUNKS_API
+		TArray<FChunkBatchRequestEntry> ChunkBatchRequests;
+		TOptional<uint64> EmptyModTag;
+		ChunkBatchRequests.Emplace(FChunkBatchRequestEntry{ChunkId, Offset, Size, EmptyModTag});
+
+		FIoStatus ResultStatus = ReadChunkBatchRequest(ChunkBatchRequests, [&](FIoChunkId Id, EStorageServerContentType MimeType, FIoBuffer Data, uint64 ModTag)
+		{
+			Data.MakeOwned();
+			ResultTuple = IStorageServerHttpClient::FResult(Data, MimeType);
+		});
+
+		if (!ResultStatus.IsOk())
+		{
+			ResultTuple = IStorageServerHttpClient::FResult(ResultStatus, EStorageServerContentType::Unknown);
+		}
+#else
 		TRACE_CPUPROFILER_EVENT_SCOPE(FStorageServerConnection::ReadChunkRequest::Http);
 
 		TAnsiStringBuilder<256> ResourceBuilder;
 		BuildReadChunkRequestUrl(ResourceBuilder, ChunkId, Offset, Size);
 
 		ResultTuple = HttpClient->RequestSync(*ResourceBuilder);
+#endif
 	}
 
 	TIoStatusOr<FIoBuffer> ResultBuffer = ReadChunkRequestProcessHttpResult(ResultTuple, Offset, Size, OptDestination, bHardwareTargetBuffer);
@@ -308,29 +326,55 @@ TIoStatusOr<FIoBuffer> FStorageServerConnection::ReadChunkRequest(
 	return ResultBuffer;
 }
 
-TIoStatusOr<FIoBuffer> FStorageServerConnection::ReadChunkBatchRequest(const TArray<FChunkBatchRequestEntry>& Chunks)
+#if HAS_STORAGE_SERVER_RPC_GETCHUNKS_API
+FIoStatus FStorageServerConnection::ReadChunkBatchRequest(const TArray<FChunkBatchRequestEntry>& Chunks, TFunctionRef<void(FIoChunkId Id, EStorageServerContentType MimeType, FIoBuffer Data, uint64 ModTag)> OnResponse, bool bSkipData)
 {
 	TRACE_CPUPROFILER_EVENT_SCOPE(FStorageServerConnection::ReadChunkBatchRequest);
 
+	if (Chunks.Num() == 0)
+	{
+		return FIoStatus::Ok;
+	}
+	else if (Chunks.Num() > 1)
+	{
+		// TODO Implement multiple requests to same chunkid in same batch.
+
+		TSet<FIoChunkId> RequestedChunkIds;
+		for (const FChunkBatchRequestEntry& Chunk: Chunks)
+		{
+			if (RequestedChunkIds.Contains(Chunk.Id))
+			{
+				UE_LOG(LogStorageServerConnection, Fatal, TEXT("Multiple requests to the same chunkid are not supported"));
+				return FIoStatus::Unknown;
+			}
+
+			RequestedChunkIds.Add(Chunk.Id);
+		}
+	}
+
 	const double StartTime = FPlatformTime::Seconds();
 
 	FCbWriter Writer;
-
 	Writer.BeginObject();
-
 	Writer << "method" << "getchunks";
 
 	Writer.BeginObject("Request");
+	if (bSkipData)
+	{
+		Writer << "SkipData" << bSkipData;
+	}
 	Writer.BeginArray("Chunks");
 
 	for (const FChunkBatchRequestEntry& Chunk : Chunks)
 	{
 		Writer.BeginObject();
-
 		Writer << "Offset" << Chunk.Offset;
 		Writer << "Size" << Chunk.Size;
 		Writer << "Oid" << Chunk.Id;
-
+		if (Chunk.ModTag.IsSet())
+		{
+			Writer << "ModTag" << *Chunk.ModTag;
+		}
 		Writer.EndObject();
 	}
 
@@ -341,27 +385,131 @@ TIoStatusOr<FIoBuffer> FStorageServerConnection::ReadChunkBatchRequest(const TAr
 	TAnsiStringBuilder<256> Uri;
 	Uri.Append(BaseURI).Append("/rpc");
 
-	FIoBuffer Payload(Writer.GetSaveSize());
-	Writer.Save(Payload.GetMutableView());
+	FIoBuffer PostPayload(Writer.GetSaveSize());
+	Writer.Save(PostPayload.GetMutableView());
 
 	IStorageServerHttpClient::FResult HttpResult = HttpClient->RequestSync(
 		*Uri,
 		EStorageServerContentType::CompressedBinary,
 		"POST",
-		Payload,
+		PostPayload,
 		EStorageServerContentType::CbObject
 	);
 
 	TIoStatusOr<FIoBuffer> Result = HttpResult.Get<0>();
 
-	if (Result.IsOk())
+	if (!Result.IsOk())
 	{
-		const double Duration = FPlatformTime::Seconds() - StartTime;
-		AddTimingInstance(Duration, Result.ValueOrDie().GetSize());
+		return Result.Status();
 	}
 
-	return Result;
+	const double Duration = FPlatformTime::Seconds() - StartTime;
+	AddTimingInstance(Duration, Result.ValueOrDie().GetSize());
+
+	// TODO Replace parsing code with FCbPackage.
+
+	struct ResponseHeader
+	{
+		uint32 Magic;
+		uint32 AttachmentCount;
+		uint32 Reserved[2];
+	} Header;
+
+	FMemoryReaderView View(Result.ValueOrDie().GetView());
+	View.Serialize(&Header, sizeof(Header));
+
+	struct AttachmentInfo
+	{
+		uint64 PayloadSize;
+		uint32 Flags;
+		FIoHash Hash;
+
+		enum
+		{
+			IsCompressed = (1u << 0),
+			IsObject = (1u << 1),
+		};
+	};
+
+	if (Header.Magic != 0xaa77aacc)
+	{
+		UE_LOG(LogStorageServerConnection, Fatal, TEXT("Read incorrect header magic from server response - expected %x, got %x"), 0xaa77aacc, Header.Magic);
+		return FIoStatus::Unknown;
+	}
+
+	TArray<AttachmentInfo> Info;
+	Info.AddUninitialized(Header.AttachmentCount + 1); // One extra for the CbPackage root object.
+
+	TMap<FIoHash, int32> AttachmentHashToIndex;
+
+	for (int32 i = 0; i < Info.Num(); ++i)
+	{
+		View.Serialize(&Info[i], sizeof(Info[i]));
+
+		AttachmentHashToIndex.Add(Info[i].Hash, i);
+	}
+
+	FCbObjectView Root(Result.ValueOrDie().GetData() + View.Tell());
+
+	TArray<uint64> InfoOffsets;
+	InfoOffsets.AddUninitialized(Header.AttachmentCount + 1);
+	for (int32 i = 0; i < Info.Num(); ++i)
+	{
+		InfoOffsets[i] = (i == 0) ? View.Tell() : InfoOffsets[i - 1] + Info[i - 1].PayloadSize;
+	}
+
+	FCbArrayView ResponseChunks = Root["Chunks"].AsArrayView();
+
+	uint64 ResponseCount = 0;
+	for (FCbFieldView Field : ResponseChunks)
+	{
+		const FCbObjectView Chunk = Field.AsObjectView();
+		const FCbObjectId Id = Chunk["Id"].AsObjectId();
+		const FIoHash RawHash = Chunk["RawHash"].AsHash();
+		const FIoHash FragmentHash = Chunk["FragmentHash"].AsHash();
+		const uint64 ModTag = Chunk["ModTag"].AsUInt64();
+
+		if (!ensureMsgf(RawHash != FIoHash() || FragmentHash != FIoHash(), TEXT("Failed to find hash in chunk info returned from server")))
+		{
+			continue;
+		}
+
+		const FIoHash& AttachmentHash = RawHash == FIoHash() ? FragmentHash : RawHash;
+
+		int32* InfoIndex = AttachmentHashToIndex.Find(AttachmentHash);
+		if (!ensureMsgf(InfoIndex, TEXT("Failed to find hash in attachments returned from server")))
+		{
+			continue;
+		}
+
+		FIoChunkId ChunkId;
+		ChunkId.Set(Id.GetView());
+
+		const AttachmentInfo& Attachment = Info[*InfoIndex];
+		ensureAlways(Attachment.Hash == AttachmentHash);
+
+		EStorageServerContentType MimeType = EStorageServerContentType::CompressedBinary;
+		if (Attachment.Flags & AttachmentInfo::IsCompressed)
+		{
+			MimeType = EStorageServerContentType::CompressedBinary;
+		}
+		else if (Attachment.Flags & AttachmentInfo::IsObject)
+		{
+			MimeType = EStorageServerContentType::CbObject;
+		}
+		else
+		{
+			MimeType = EStorageServerContentType::Binary;
+		}
+
+		FIoBuffer Data(FIoBuffer::Wrap, Result.ValueOrDie().GetData() + InfoOffsets[*InfoIndex], Attachment.PayloadSize);
+		OnResponse(ChunkId, MimeType, Data, ModTag);
+		ResponseCount++;
+	}
+
+	return ResponseCount ? FIoStatus::Ok : EIoErrorCode::NotFound;
 }
+#endif
 
 void FStorageServerConnection::ReadChunkRequestAsync(
 	const FIoChunkId& ChunkId,
@@ -786,7 +934,7 @@ TIoStatusOr<FIoBuffer> FStorageServerConnection::ReadChunkRequestProcessHttpResu
 	EStorageServerContentType MimeType = ResultTuple.Get<1>();
 	if (!Result.IsOk())
 	{
-		UE_LOG(LogStorageServerConnection, Fatal, TEXT("Failed read chunk from storage server. '%s'  Offset:%ull  Size:%ull"), *Result.Status().ToString(), Offset, Size);
+		UE_LOG(LogStorageServerConnection, Warning, TEXT("Failed read chunk from storage server. '%s'  Offset:%ull  Size:%ull"), *Result.Status().ToString(), Offset, Size);
 		return Result.Status();
 	}
 
@@ -808,6 +956,7 @@ TIoStatusOr<FIoBuffer> FStorageServerConnection::ReadChunkRequestProcessHttpResu
 		}
 		else
 		{
+			Buffer.MakeOwned();
 			return Buffer;
 		}
 	}
diff --git a/Engine/Source/Runtime/StorageServerClient/Private/StorageServerConnection.h b/Engine/Source/Runtime/StorageServerClient/Private/StorageServerConnection.h
--- a/Engine/Source/Runtime/StorageServerClient/Private/StorageServerConnection.h
+++ b/Engine/Source/Runtime/StorageServerClient/Private/StorageServerConnection.h
@@ -25,6 +25,10 @@
 	#define HAS_STORAGE_SERVER_PLATFORM_CACHE_STRATEGY 0
 #endif
 
+#ifndef HAS_STORAGE_SERVER_RPC_GETCHUNKS_API // zen server 5.5.16 introduced a new API to request chunks
+	#define HAS_STORAGE_SERVER_RPC_GETCHUNKS_API 1
+#endif
+
 #if !UE_BUILD_SHIPPING
 
 DECLARE_LOG_CATEGORY_EXTERN(LogStorageServerConnection, Log, All);
@@ -51,15 +55,6 @@ public:
 		const bool bHardwareTargetBuffer
 	);
 
-	// Matches RequestChunkEntry in zenserver httprojectstore.cpp
-	struct FChunkBatchRequestEntry {
-		FIoChunkId Id;
-		uint64 Offset;
-		uint64 Size;
-	};
-
-	TIoStatusOr<FIoBuffer> ReadChunkBatchRequest(const TArray<FChunkBatchRequestEntry>& Chunks);
-	
 	void ReadChunkRequestAsync(
 		const FIoChunkId& ChunkId,
 		const uint64 Offset,
@@ -69,6 +64,22 @@ public:
 		TFunctionRef<void(TIoStatusOr<FIoBuffer> Data)> OnResponse
 	);
 
+#if HAS_STORAGE_SERVER_RPC_GETCHUNKS_API
+	// Matches input parameters for ProjectStore::GetChunks in zen server
+	struct FChunkBatchRequestEntry {
+		FIoChunkId Id;
+		uint64 Offset;
+		uint64 Size;
+		TOptional<uint64> ModTag;
+	};
+
+	FIoStatus ReadChunkBatchRequest(
+		const TArray<FChunkBatchRequestEntry>& Chunks,
+		TFunctionRef<void(FIoChunkId Id, EStorageServerContentType MimeType, FIoBuffer Data, uint64 ModTag)> OnResponse,
+		bool bSkipData = false // if bSkipData is true, then OnResponse will only be called for requested chunks that have either empty or different modtags
+	);
+#endif
+
 	FStringView GetHostAddr() const
 	{
 		return CurrentHostAddr;
diff --git a/Engine/Source/Runtime/StorageServerClient/Private/StorageServerPlatformFile.cpp b/Engine/Source/Runtime/StorageServerClient/Private/StorageServerPlatformFile.cpp
--- a/Engine/Source/Runtime/StorageServerClient/Private/StorageServerPlatformFile.cpp
+++ b/Engine/Source/Runtime/StorageServerClient/Private/StorageServerPlatformFile.cpp
@@ -337,83 +337,31 @@ private:
 		return 0;
 	}
 
-	void HandleChunkBatchResponse(const FIoBuffer& Payload)
+	void HandleChunkBatchResponse(FIoChunkId ChunkId, EStorageServerContentType MimeType, FIoBuffer Data, uint64 ModTag)
 	{
-		struct ResponseHeader
-		{
-			uint32 Magic;
-			uint32 AttachmentCount;
-			uint32 Reserved[2];
-		} Header;
-
-		FMemoryReaderView View(Payload.GetView());
-		View.Serialize(&Header, sizeof(Header));
-
-		struct AttachmentInfo
-		{
-			uint64 PayloadSize;
-			uint32 Flags;
-			FIoHash Hash;
-		};
-
-		if (Header.Magic != 0xaa77aacc)
-		{
-			UE_LOG(LogStorageServerPlatformFile, Warning, TEXT("Read incorrect header magic from server response - expected %x, got %x"), 0xaa77aacc, Header.Magic);
-			return;
-		}
-
-		TArray<AttachmentInfo> Info;
-		Info.AddUninitialized(Header.AttachmentCount + 1); // one extra for the CbPackage root object
-
-		for (int32 i = 0; i < Info.Num(); ++i)
-		{
-			View.Serialize(&Info[i], sizeof(Info[i]));
-		}
-
-		FCbObjectView Root(Payload.GetData() + View.Tell());
-		FCbArrayView Chunks = Root["Chunks"].AsArrayView();
-
-		TMap<FIoHash, TArray<uint8>> AttachmentData;
+		TArray<uint8> DecompressedData;
 
-		uint64 Offset = View.Tell() + Root.GetSize();
-		for (int32 i = 1; i < Info.Num(); ++i) // skip the attachment info for the root object
+		// TODO move decompression to StorageServerConnection.
+		if (MimeType == EStorageServerContentType::CompressedBinary)
 		{
-			FCompressedBuffer Buffer = FCompressedBuffer::FromCompressed(FSharedBuffer::MakeView(Payload.GetData() + Offset, Info[i].PayloadSize));
+			FCompressedBuffer Buffer = FCompressedBuffer::FromCompressed(FSharedBuffer::MakeView(Data.GetView()));
 			FCompressedBufferReader Reader(Buffer);
 
-			TArray<uint8> Data;
-			Data.AddUninitialized(Reader.GetRawSize());
+			DecompressedData.AddUninitialized(Reader.GetRawSize());
 
-			if (!ensureMsgf(Reader.TryDecompressTo(FMutableMemoryView(Data.GetData(), Data.Num())), TEXT("Failed to decompress data from server response")))
+			if (!ensureMsgf(Reader.TryDecompressTo(FMutableMemoryView(DecompressedData.GetData(), DecompressedData.Num())), TEXT("Failed to decompress data from server response")))
 			{
-				continue;
+				return;
 			}
-
-			AttachmentData.Add(Info[i].Hash, MoveTemp(Data));
-			Offset += Info[i].PayloadSize;
 		}
-
-		for (FCbFieldView Field : Chunks)
+		else
 		{
-			FCbObjectView Chunk = Field.AsObjectView();
-			FCbObjectId Id = Chunk["Id"].AsObjectId();
-			FIoHash RawHash = Chunk["RawHash"].AsHash();
-			FIoHash FragmentHash = Chunk["FragmentHash"].AsHash();
-
-			if (!ensureMsgf(RawHash != FIoHash() || FragmentHash != FIoHash(), TEXT("Failed to find hash in chunk info returned from server")))
-			{
-				continue;
-			}
+			DecompressedData.Append(Data.GetData(), Data.GetSize());
+		}
 
-			FIoChunkId ChunkId;
-			ChunkId.Set(Id.GetView());
-			TArray<uint8>* Data = AttachmentData.Find(RawHash == FIoHash() ? FragmentHash : RawHash);
+		FWriteScopeLock Lock(CacheLock);
+		Cache.Add(ChunkId, DecompressedData);
 
-			{
-				FWriteScopeLock Lock(CacheLock);
-				Cache.Add(ChunkId, *Data);
-			}
-		}
 	}
 
 	void PrecacheItems()
@@ -445,15 +393,15 @@ private:
 			return;
 		}
 
-		TIoStatusOr<FIoBuffer> Response = Connection.ReadChunkBatchRequest(Items);	
+		FIoStatus ResponseStatus = Connection.ReadChunkBatchRequest(Items, [&](FIoChunkId Id, EStorageServerContentType MimeType, FIoBuffer Data, uint64 ModTag)
+		{
+			HandleChunkBatchResponse(Id, MimeType, Data, ModTag);
+		});
 
-		if (!ensureMsgf(Response.IsOk(), TEXT("Failed to read chunk batch request from Zen server")))
+		if (!ensureMsgf(ResponseStatus.IsOk(), TEXT("Failed to read chunk batch request from Zen server")))
 		{
 			return;
 		}
-
-		FIoBuffer ResponseBuffer = Response.ValueOrDie();
-		HandleChunkBatchResponse(ResponseBuffer);
 	}
 
 	FStorageServerConnection& Connection;
