Simplified sectioned cache journal

diff --git a/Engine/Source/Runtime/StorageServerClient/Private/Cache/CacheJournalSectioned.cpp b/Engine/Source/Runtime/StorageServerClient/Private/Cache/CacheJournalSectioned.cpp
new file mode 100644
--- /dev/null
+++ b/Engine/Source/Runtime/StorageServerClient/Private/Cache/CacheJournalSectioned.cpp
@@ -0,0 +1,691 @@
+﻿// Copyright Epic Games, Inc. All Rights Reserved.
+
+#include "CacheJournalSectioned.h"
+
+#include "HAL/PlatformFile.h"
+#include "Misc/ScopeLock.h"
+#include "Serialization/MemoryWriter.h"
+
+#if !UE_BUILD_SHIPPING
+
+template<typename M>
+class TMapEnum
+{
+public:
+	M& Data;
+
+	TMapEnum(M& Item)
+		: Data(Item)
+	{
+	}
+
+	TMapEnum& operator()(typename M::KeyType const& Key, typename M::ValueType Value)
+	{
+		Data.Add(Key, Value);
+		return *this;
+	}
+};
+
+template<typename T>
+TMapEnum<T> TMapInit(T& Item)
+{
+	return TMapEnum<T>(Item);
+}
+
+
+DEFINE_LOG_CATEGORY_STATIC(LogCacheJournal, Log, All);
+
+namespace StorageServer
+{
+	static const uint64 PAGE_MAGIC = 0x5041'4745; // PAGE
+	static const uint32 JOURNAL_VERSION = 0x01;
+
+	FCacheJournalSectioned::FCacheJournalSectioned(FStringView InJournalPath)
+		: JournalFilePath(InJournalPath)
+	{
+		JournalFileHandle.Reset(IPlatformFile::GetPlatformPhysical().OpenWrite(*JournalFilePath, true, true));
+
+		LoadCacheJournal(JournalFilePath);
+	}
+
+	void FCacheJournalSectioned::Flush()
+	{
+		TRACE_CPUPROFILER_EVENT_SCOPE(FCacheJournalSectioned::Flush);
+
+		UE::TScopeLock _(CriticalSection);
+
+		if (bPagesModified)
+		{
+			FJournalHeader Header{ .Version = JOURNAL_VERSION, .NumberOfChunkPages = ChunkPages.Num(), .NumberOfChunkInfoPages = ChunkInfoPages.Num() };
+			JournalFileHandle->Seek(0);
+			JournalFileHandle->Write((uint8*)&Header, sizeof(Header));
+
+			bPagesModified = false;
+		}
+
+		for (FJournalChunkPage& Page : ChunkPages)
+		{
+			Page.Flush(JournalFileHandle.Get());
+		}
+
+		for (FJournalChunkInfoPage& Page : ChunkInfoPages)
+		{
+			Page.Flush(JournalFileHandle.Get());
+		}
+	}
+
+	void FCacheJournalSectioned::InvalidateAll()
+	{
+		UE::TScopeLock _(CriticalSection);
+
+		for (FJournalChunkPage& Page : ChunkPages)
+		{
+			Page.InvalidateAll();
+		}
+
+		for (FJournalChunkInfoPage& InfoPage : ChunkInfoPages)
+		{
+			InfoPage.InvalidateAll();
+		}
+
+		bPagesModified = true;
+	}
+
+	void FCacheJournalSectioned::Invalidate(const FIoChunkId& ChunkId)
+	{
+		UE::TScopeLock _(CriticalSection);
+
+		for (FJournalChunkPage& Page : ChunkPages)
+		{
+			Page.Invalidate(ChunkId);
+		}
+
+		for (FJournalChunkInfoPage& InfoPage : ChunkInfoPages)
+		{
+			InfoPage.Invalidate(ChunkId);
+		}
+	}
+
+	bool FCacheJournalSectioned::SetChunkInfo(const FIoChunkId& ChunkId, const FModHash* OptModHash, const int64* OptRawSize, const int32* OptRawBlockSize)
+	{
+		UE::TScopeLock _(CriticalSection);
+
+		FCacheChunkInfo ChunkInfo;
+		if (TryGetChunkInfo(ChunkId, ChunkInfo))
+		{
+			return ChunkInfo.SetChunkInfo(OptModHash, OptRawSize, OptRawBlockSize);
+		}
+
+		FJournalChunkInfoPage* Page = FindAddChunkInfoPage(ChunkId);
+		check(Page != nullptr && !Page->IsFull());
+
+		ChunkInfo.SetChunkInfo(OptModHash, OptRawSize, OptRawBlockSize);
+		return Page->Add(ChunkId, ChunkInfo) == EJournalPageResult::Ok;
+	}
+
+	bool FCacheJournalSectioned::TryGetChunkInfo(const FIoChunkId& ChunkId, FCacheChunkInfo& OutChunkInfo)
+	{
+		UE::TScopeLock _(CriticalSection);
+
+		for (FJournalChunkInfoPage& Page : ChunkInfoPages)
+		{
+			if (Page.GetEntry(ChunkId, OutChunkInfo))
+			{
+				return true;
+			}
+		}
+
+		return false;
+	}
+
+	bool FCacheJournalSectioned::AddEntry(const FIoChunkId& ChunkId, const FCacheEntry& Entry)
+	{
+		EJournalPageResult Res;
+
+		UE::TScopeLock _(CriticalSection);
+
+		FCacheEntry OutEntry;
+		if (TryGetEntry(ChunkId, Entry.ChunkOffset, Entry.ChunkSize, OutEntry))
+		{
+			return false;
+		}
+
+		FJournalChunkPage* Page = FindChunkPage(ChunkId);
+
+		Res = Page->Add(ChunkId, Entry);
+
+		return Res == EJournalPageResult::Ok;
+	}
+
+	bool FCacheJournalSectioned::TryGetEntry(const FIoChunkId& ChunkId, const uint64 ChunkOffset, const uint64 ChunkSize, FCacheEntry& OutEntry)
+	{
+		UE::TScopeLock _(CriticalSection);
+
+		for (FJournalChunkPage& Page : ChunkPages)
+		{
+			if (Page.GetEntry(ChunkId, ChunkOffset, ChunkSize, OutEntry))
+			{
+				return true;
+			}
+		}
+
+		return false;
+	}
+
+	void FCacheJournalSectioned::IterateChunkIds(TFunctionRef<void(const FIoChunkId& ChunkId)> Callback)
+	{
+		UE::TScopeLock _(CriticalSection);
+
+		for (FJournalChunkPage& Page : ChunkPages)
+		{
+			Page.IterateChunkIds(MoveTemp(Callback));
+		}
+	}
+
+	void FCacheJournalSectioned::IterateCacheEntriesForChunkId(const FIoChunkId& ChunkId, TFunctionRef<void(const FCacheEntry& Entry)> Callback)
+	{
+		UE::TScopeLock _(CriticalSection);
+
+		for (FJournalChunkPage& Page : ChunkPages)
+		{
+			Page.IterateCacheEntriesForChunkId(ChunkId, MoveTemp(Callback));
+		}
+	}
+
+	void FCacheJournalSectioned::IterateCacheEntries(TFunctionRef<void(const FIoChunkId& ChunkId, const FCacheEntry& Entry)> Callback)
+	{
+		UE::TScopeLock _(CriticalSection);
+
+		for (FJournalChunkPage& Page : ChunkPages)
+		{
+			Page.IterateCacheEntries(MoveTemp(Callback));
+		}
+	}
+
+	TArray<FIoChunkId> FCacheJournalSectioned::ApplyLatestServerChunkInfo(const TArray<FLatestServerChunkInfo>& ServerChunkInfoArray)
+	{
+		UE::TScopeLock _(CriticalSection);
+
+		TArray<FIoChunkId> InvalidatedChunks;
+		for (const FLatestServerChunkInfo& ServerChunkInfo : ServerChunkInfoArray)
+		{
+			FCacheChunkInfo ChunkInfo;
+			if (TryGetChunkInfo(ServerChunkInfo.ChunkId, ChunkInfo))
+			{
+				const bool bHashesMatch = ChunkInfo.SetChunkInfo(&ServerChunkInfo.ModHash, &ServerChunkInfo.RawSize, nullptr);
+				if (!bHashesMatch)
+				{
+					InvalidatedChunks.Add(ServerChunkInfo.ChunkId);
+				}
+			}
+			else
+			{
+				SetChunkInfo(ServerChunkInfo.ChunkId, &ServerChunkInfo.ModHash, &ServerChunkInfo.RawSize, nullptr);
+			}
+		}
+
+		return InvalidatedChunks;
+	}
+
+	void FCacheJournalSectioned::LoadCacheJournal(FStringView JournalPath)
+	{
+		TRACE_CPUPROFILER_EVENT_SCOPE(FCacheJournalSectioned::LoadCacheJournal);
+
+		UE::TScopeLock _(CriticalSection);
+
+		check(JournalFileHandle.IsValid());
+
+		const int64 FileSize = JournalFileHandle->Size();
+		if (FileSize < sizeof(FJournalHeader))
+		{
+			return;
+		}
+
+		FJournalHeader Header;
+		JournalFileHandle->ReadAt((uint8*)&Header, sizeof(Header), 0);
+		if (Header.Version != JOURNAL_VERSION)
+		{
+			UE_LOG(LogCacheJournal, Error, TEXT("Journal file is of different version %d, supported %d"), Header.Version, JOURNAL_VERSION);
+			return;
+		}
+
+		int64 ReadCursor = sizeof(Header);
+		int32 ChunkEntryCount = 0;
+		int32 ChunkInfoEntryCount = 0;
+		int32 TotalMBs = 0;
+
+		TMap<EPageType, FString> TypeArray;
+		TMapInit(TypeArray)
+			(EPageType::Chunk, TEXT("Chunk"))
+			(EPageType::ChunkInfo, TEXT("ChunkInfo"));
+
+		while ((uint32)(ReadCursor + sizeof(FJournalPageDescriptor)) < FileSize)
+		{
+			auto ReadAndProgress = [&JournalFileHandle = JournalFileHandle, &ReadCursor](uint8* Destination, uint32 Size) -> bool
+			{
+				if (!JournalFileHandle->ReadAt(Destination, Size, ReadCursor))
+				{
+					return false;
+				}
+				ReadCursor += Size;
+				return true;
+			};
+
+			int32 PageFilePos = ReadCursor;
+			FJournalPageDescriptor SectionDescriptor;
+
+			if (!ReadAndProgress((uint8*)&SectionDescriptor, sizeof(SectionDescriptor)))
+			{
+				UE_LOG(LogCacheJournal, Error, TEXT("Failed to read section from journal file"));
+				InvalidateAll();
+				return;
+			}
+
+			if (SectionDescriptor.Magic != PAGE_MAGIC)
+			{
+				UE_LOG(LogCacheJournal, Error, TEXT("Section has wrong magic number, the journal is corrupt"));
+				InvalidateAll();
+				return;
+			}
+
+			if (SectionDescriptor.Type == EPageType::Chunk)
+			{
+				int32 EntriesCount = SectionDescriptor.Size / sizeof(FJournalStoreEntry);
+				TArray<FJournalStoreEntry> Entries;
+				Entries.AddUninitialized(EntriesCount);
+
+				ChunkEntryCount += EntriesCount;
+
+				if (!ReadAndProgress((uint8*)Entries.GetData(), SectionDescriptor.Size))
+				{
+					UE_LOG(LogCacheJournal, Error, TEXT("Failed to read chunks from section in journal file"));
+					InvalidateAll();
+					return;
+				}
+
+				int64 AllChunksSize = 0;
+				int64 NumberOfNoSizeChunks = 0;
+				for (FJournalStoreEntry& Entry : Entries)
+				{
+					if (Entry.Entry.ChunkSize <= 0)
+					{
+						NumberOfNoSizeChunks++;
+						continue;
+					}
+
+					AllChunksSize += Entry.Entry.ChunkSize;
+				}
+
+				const int32 SizeInMb = AllChunksSize >> 20;
+				TotalMBs += SizeInMb;
+				UE_LOG(LogCacheJournal, Log, TEXT("Found %s page with %d entries of total size %d, Number of chunks with size -1 %d"), 
+					*TypeArray[SectionDescriptor.Type], EntriesCount, SizeInMb, NumberOfNoSizeChunks);
+
+				FJournalChunkPage& Page = ChunkPages.Add_GetRef({ GMaxSectionEntries });
+				Page.SetFilePos(PageFilePos);
+				Page << MoveTemp(Entries);
+			}
+			else if (SectionDescriptor.Type == EPageType::ChunkInfo)
+			{
+				TArray<FJournalChunkInfoPage::FChunkInfoEntry> ChunkInfoEntries;
+				int32 EntriesCount = SectionDescriptor.Size / sizeof(FJournalChunkInfoPage::FChunkInfoEntry);
+				ChunkInfoEntries.AddUninitialized(EntriesCount);
+
+				ChunkInfoEntryCount += EntriesCount;
+
+				if (!ReadAndProgress((uint8*)ChunkInfoEntries.GetData(), SectionDescriptor.Size))
+				{
+					UE_LOG(LogCacheJournal, Error, TEXT("Failed to read chunk infos from section in journal file"));
+					InvalidateAll();
+					return;
+				}
+
+				FJournalChunkInfoPage& Page = ChunkInfoPages.Add_GetRef({ GMaxSectionEntries });
+				Page.SetFilePos(PageFilePos);
+				Page << MoveTemp(ChunkInfoEntries);
+			}
+			
+			int32 BytesLeft = SectionDescriptor.MaxSize - SectionDescriptor.Size;
+			if (BytesLeft > 0)
+			{
+				ReadCursor += BytesLeft;
+			}
+		}
+
+		check(Header.NumberOfChunkInfoPages == ChunkInfoPages.Num());
+		check(Header.NumberOfChunkPages == ChunkPages.Num());
+
+		JournalSectionCursor = ReadCursor;
+		NextAvailableJournalPos = ReadCursor;
+		UE_LOG(LogCacheJournal, Display, TEXT("Cache journal loaded. Imported %d chunk pages with %d entries and %d chunk info pages with %d info entries. Total size of data referenced by journal %d MB"),
+			ChunkPages.Num(), ChunkEntryCount, ChunkInfoPages.Num(), ChunkInfoEntryCount, TotalMBs);
+	}
+
+	FJournalChunkInfoPage* FCacheJournalSectioned::FindAddChunkInfoPage(const FIoChunkId& ChunkId)
+	{
+		for (FJournalChunkInfoPage& Page : ChunkInfoPages)
+		{
+			check(!Page.Contains(ChunkId));
+		}
+
+		for (FJournalChunkInfoPage& Page : ChunkInfoPages)
+		{
+			if (!Page.IsFull())
+			{
+				return &Page;
+			}
+		}
+
+		FJournalChunkInfoPage* ResPage = &ChunkInfoPages.Add_GetRef({ GMaxSectionEntries});
+		ResPage->SetFilePos(ReserveNextJournalFilePos(ResPage->GetPageMaxSize()));
+		bPagesModified = true;
+		return ResPage;
+	}
+
+	FJournalChunkPage* FCacheJournalSectioned::FindChunkPage(const FIoChunkId& ChunkId)
+	{
+		FJournalChunkPage* LastPage = nullptr;
+		for (FJournalChunkPage& Page : ChunkPages)
+		{
+			if (!Page.IsFull())
+			{
+				LastPage = &Page;
+
+				if (Page.Contains(ChunkId))
+				{
+					return &Page;
+				}
+			}
+		}
+
+		if (!LastPage)
+		{
+			LastPage = &ChunkPages.Add_GetRef({ GMaxSectionEntries });
+			LastPage->SetFilePos(ReserveNextJournalFilePos(LastPage->GetPageMaxSize()));
+			bPagesModified = true;
+		}
+
+		return LastPage;
+	}
+	
+	FJournalChunkInfoPage::FJournalChunkInfoPage(int32 InMaxEntries)
+		: FJournalPageBase(InMaxEntries)
+	{
+	}
+
+	EJournalPageResult FJournalChunkInfoPage::Add(const FIoChunkId& ChunkId, const FCacheChunkInfo& Entry)
+	{
+		if (Entries.Num() >= MaxEntryCount)
+		{
+			return EJournalPageResult::PageFull;
+		}
+
+		FCacheChunkInfo* ChunkInfo = Entries.Find(ChunkId);
+		if (ChunkInfo)
+		{
+			FModHash PreviousHash = ChunkInfo->ModHash;
+			const bool bHashesMatch = ChunkInfo->SetChunkInfo(&Entry.ModHash, &Entry.RawSize, &Entry.RawBlockSize);
+			if (!bHashesMatch)
+			{
+				bDirty = true;
+			}
+			return EJournalPageResult::EntryAlreadyExists;
+		}
+
+		Entries.Add(ChunkId, Entry);
+		bDirty = true;
+
+		return EJournalPageResult::Ok;
+	}
+
+	bool FJournalChunkInfoPage::GetEntry(const FIoChunkId& ChunkId, FCacheChunkInfo& Info)
+	{
+		FCacheChunkInfo* ChunkInfo = Entries.Find(ChunkId);
+		if (!ChunkInfo)
+		{
+			return false;
+		}
+
+		Info = *ChunkInfo;
+		return true;
+	}
+
+	bool FJournalChunkInfoPage::Contains(const FIoChunkId& ChunkId) const
+	{
+		return Entries.Contains(ChunkId);
+	}
+
+	void FJournalChunkInfoPage::Invalidate(const FIoChunkId& ChunkId)
+	{
+		bDirty = Entries.Remove(ChunkId) > 0;
+	}
+
+	void FJournalChunkInfoPage::InvalidateAll()
+	{
+		bDirty = Entries.Num() > 0;
+		Entries.Empty();
+	}
+
+	bool FJournalChunkInfoPage::Flush(IFileHandle* Handle)
+	{
+		TRACE_CPUPROFILER_EVENT_SCOPE(FJournalChunkInfoPage::Flush);
+
+		if (!bDirty)
+		{
+			return true;
+		}
+
+		FJournalPageDescriptor Descriptor;
+		Descriptor.Magic = PAGE_MAGIC;
+		Descriptor.Type = EPageType::ChunkInfo;
+		Descriptor.MaxEntryCount = MaxEntryCount;
+
+		TArray<FChunkInfoEntry> OutputData;
+		OutputData.Reserve(Entries.Num());
+
+		for (decltype(Entries)::TIterator It = Entries.CreateIterator(); It; ++It)
+		{
+			OutputData.Add({ .ChunkId = It.Key(), .ChunkInfo = It.Value() });
+		}
+
+		Descriptor.Size = OutputData.Num() * OutputData.GetTypeSize();
+		Descriptor.MaxSize = MaxEntryCount * OutputData.GetTypeSize();
+
+		int64 CurrentPos = Handle->Tell();
+		Handle->Seek(FilePos);
+		if (!Handle->Write((uint8*)&Descriptor, sizeof(Descriptor)))
+		{
+			UE_LOG(LogCacheJournal, Error, TEXT("Failed to write chunk info page descriptor"));
+			return false;
+		}
+
+		if (!Handle->Write((uint8*)OutputData.GetData(), OutputData.Num() * OutputData.GetTypeSize()))
+		{
+			UE_LOG(LogCacheJournal, Error, TEXT("Failed to write chunk info entries"));
+			return false;
+		}
+		Handle->Seek(CurrentPos);
+
+		bDirty = false;
+		return true;
+	}
+
+	FJournalChunkInfoPage& FJournalChunkInfoPage::operator<<(const TArray<FChunkInfoEntry>& InEntries)
+	{
+		for (const FChunkInfoEntry& Entry : InEntries)
+		{
+			Entries.Add(Entry.ChunkId, Entry.ChunkInfo);
+		}
+
+		return *this;
+	}
+
+	FJournalChunkPage::FJournalChunkPage(int32 InMaxEntries)
+		: FJournalPageBase(InMaxEntries)
+	{
+	}
+
+	EJournalPageResult FJournalChunkPage::Add(const FIoChunkId& ChunkId, const FCacheEntry& Entry)
+	{
+		if (Entries.Num() >= MaxEntryCount)
+		{
+			return EJournalPageResult::PageFull;
+		}
+
+		FCacheEntry Out;
+		if (GetEntry(ChunkId, Entry.ChunkOffset, Entry.ChunkSize, Out))
+		{
+			return EJournalPageResult::EntryAlreadyExists;
+		}
+
+		FJournalStoreEntry StoreEntry;
+		StoreEntry.ChunkId = ChunkId;
+		StoreEntry.Entry = Entry;
+		StoreEntry.bValid = true;
+		int32 Idx = Entries.Add(StoreEntry);
+		FIntervalMap& IntervalMap = ChunkMap.FindOrAdd(ChunkId);
+		IntervalMap.Add(Entry.GetChunkInterval(), Idx);
+
+		bDirty = true;
+		return EJournalPageResult::Ok;
+	}
+
+	bool FJournalChunkPage::GetEntry(const FIoChunkId& ChunkId, const uint64 ChunkOffset, const uint64 ChunkSize, FCacheEntry& OutEntry)
+	{
+		FIntervalMap* IntervalMap = ChunkMap.Find(ChunkId);
+		if (IntervalMap)
+		{
+			int32* EntryIdx = IntervalMap->Find(TInterval<uint64>(ChunkOffset, ChunkSize));
+			if (EntryIdx)
+			{
+				OutEntry = Entries[*EntryIdx].Entry;
+				return true;
+			}
+		}
+
+		return false;
+	}
+
+	void FJournalChunkPage::IterateChunkIds(TFunctionRef<void(const FIoChunkId& ChunkId)> Callback)
+	{
+		for (const FJournalStoreEntry& Entry : Entries)
+		{
+			Callback(Entry.ChunkId);
+		}
+	}
+
+	void FJournalChunkPage::IterateCacheEntriesForChunkId(const FIoChunkId& ChunkId, TFunctionRef<void(const FCacheEntry& Entry)> Callback)
+	{
+		FIntervalMap* IntervalMap = ChunkMap.Find(ChunkId);
+		if (IntervalMap)
+		{
+			for (auto It = IntervalMap->CreateConstIterator(); It; ++It)
+			{
+				Callback(Entries[It.Value()].Entry);
+			}
+		}
+	}
+
+	void FJournalChunkPage::IterateCacheEntries(TFunctionRef<void(const FIoChunkId& ChunkId, const FCacheEntry& Entry)> Callback)
+	{
+		for (FJournalStoreEntry& Entry : Entries)
+		{
+			Callback(Entry.ChunkId, Entry.Entry);
+		}
+	}
+
+	bool FJournalChunkPage::Flush(IFileHandle* JournalFile)
+	{
+		TRACE_CPUPROFILER_EVENT_SCOPE(FJournalChunkPage::Flush);
+
+		if (!bDirty)
+		{
+			return true;
+		}
+
+		int64 CurrentPos = JournalFile->Tell();
+
+		FJournalPageDescriptor Descriptor;
+		Descriptor.Magic = PAGE_MAGIC;
+		Descriptor.Type = EPageType::Chunk;
+		Descriptor.MaxEntryCount = MaxEntryCount;
+		Descriptor.Size = Entries.Num() * Entries.GetTypeSize();
+		Descriptor.MaxSize = MaxEntryCount * Entries.GetTypeSize();
+		
+		if (bContainsInvalidEntries)
+		{
+			TArray<FJournalStoreEntry> ValidEntries;
+			ValidEntries.Reserve(Entries.Num());
+
+			const int32 EntriesCount = Entries.Num();
+			for (int32 Idx = 0; Idx < EntriesCount; ++Idx)
+			{
+				if (Entries[Idx].bValid)
+				{
+					ValidEntries.Add(Entries[Idx]);
+				}
+			}
+
+			Entries = MoveTemp(ValidEntries);
+			bContainsInvalidEntries = false;
+		}
+
+		JournalFile->Seek(FilePos);
+		JournalFile->Write((uint8*)&Descriptor, sizeof(Descriptor));
+		JournalFile->Write((uint8*)Entries.GetData(), Descriptor.Size);
+		JournalFile->Seek(CurrentPos);
+
+		bDirty = false;
+		return true;
+	}
+
+	void FJournalChunkPage::Invalidate(const FIoChunkId& ChunkId)
+	{
+		FIntervalMap* IntervalMap = ChunkMap.Find(ChunkId);
+		if (!IntervalMap)
+		{
+			return;
+		}
+
+		for (auto It = IntervalMap->CreateIterator(); It; ++It)
+		{
+			Entries[It.Value()].bValid = false;
+		}
+
+		ChunkMap.Remove(ChunkId);
+		bDirty = true;
+		bContainsInvalidEntries = true;
+	}
+
+	void FJournalChunkPage::InvalidateAll()
+	{
+		Entries.Empty();
+		ChunkMap.Empty();
+		bDirty = true;
+	}
+
+	FJournalChunkPage& FJournalChunkPage::operator<<(TArray<FJournalStoreEntry>&& InEntries)
+	{
+		Entries = MoveTemp(InEntries);
+		ChunkMap.Reset();
+
+		const int32 EntryCount = Entries.Num();
+		for (int32 Idx = 0; Idx < EntryCount; ++Idx)
+		{
+			FJournalStoreEntry& StoreEntry = Entries[Idx];
+			if (!StoreEntry.bValid)
+			{
+				bContainsInvalidEntries = true;
+				continue;
+			}
+
+			FIntervalMap& IntervalMap = ChunkMap.FindOrAdd(StoreEntry.ChunkId);
+			IntervalMap.Add(StoreEntry.Entry.GetChunkInterval(), Idx);
+		}
+
+		return *this;
+	}
+
+}
+
+#endif // !UE_BUILD_SHIPPING
\ No newline at end of file
diff --git a/Engine/Source/Runtime/StorageServerClient/Private/Cache/CacheJournalSectioned.h b/Engine/Source/Runtime/StorageServerClient/Private/Cache/CacheJournalSectioned.h
new file mode 100644
--- /dev/null
+++ b/Engine/Source/Runtime/StorageServerClient/Private/Cache/CacheJournalSectioned.h
@@ -0,0 +1,238 @@
+﻿// Copyright Epic Games, Inc. All Rights Reserved.
+
+#pragma once
+
+#include "CacheJournal.h"
+
+#include "CoreMinimal.h"
+#include "Serialization/MemoryReader.h"
+
+#if !UE_BUILD_SHIPPING
+
+namespace StorageServer
+{
+	inline constexpr uint32 GMaxSectionEntries = 1000;
+	inline constexpr uint32 JOURNAL_PAGE_SIZE = 128 * 1024;
+
+	enum class EPageType : uint8
+	{
+		Chunk,
+		ChunkInfo
+	};
+
+	struct FJournalChunkInfoEntry
+	{
+		FIoChunkId ChunkId;
+		FCacheChunkInfo ChunkInfo;
+	};
+
+	struct FPageDescriptor
+	{
+		EPageType Type;
+		int32 NumEntries;
+		uint8* Data;
+	};
+
+	struct FJournalStoreEntry
+	{
+		FIoChunkId ChunkId;
+		FCacheEntry Entry;
+		bool bValid;
+
+		bool operator<(const FJournalStoreEntry& Other) const
+		{
+			return ChunkId < Other.ChunkId;
+		}
+	};
+
+	enum class EJournalPageResult
+	{
+		Ok,
+		EntryAlreadyExists,
+		PageFull,
+		EntryNotFound
+	};
+
+	struct FJournalHeader
+	{
+		uint32 Version;
+		int32 NumberOfChunkPages;
+		int32 NumberOfChunkInfoPages;
+	};
+
+	struct FJournalPageDescriptor
+	{
+		uint32 Magic;
+		EPageType Type;
+		uint32 MaxEntryCount;
+		uint32 MaxSize;
+		uint32 Size;
+	};
+
+	class FJournalPageBase
+	{
+	public:
+		FJournalPageBase(int32 InMaxEntryCount)
+			: MaxEntryCount(InMaxEntryCount)
+		{
+		}
+
+		virtual ~FJournalPageBase()
+		{
+		}
+
+		[[nodiscard]] virtual bool Contains(const FIoChunkId&) const = 0;
+		[[nodiscard]] virtual bool IsFull() const = 0;
+
+		[[nodiscard]] virtual int32 GetPageMaxSize() const = 0;
+		void SetFilePos(int64 FileCursor)
+		{
+			FilePos = FileCursor;
+		}
+
+		virtual bool Flush(IFileHandle* JournalFile) = 0;
+
+	protected:
+		int64 FilePos;
+		int32 MaxEntryCount;
+	};
+
+	class FJournalChunkInfoPage : public FJournalPageBase
+	{
+	public:
+		FJournalChunkInfoPage(int32 InMaxEntries);
+
+		EJournalPageResult Add(const FIoChunkId& ChunkId, const FCacheChunkInfo& Entry);
+		bool GetEntry(const FIoChunkId& ChunkId, FCacheChunkInfo& Info);
+		bool Contains(const FIoChunkId& ChunkId) const override;
+		bool IsFull() const override 
+		{
+			return Entries.Num() == MaxEntryCount;
+		}
+
+		void Invalidate(const FIoChunkId& ChunkId);
+		void InvalidateAll();
+
+		int32 GetPageMaxSize() const override
+		{
+			return MaxEntryCount * sizeof(FChunkInfoEntry);
+		}
+		
+		bool Flush(IFileHandle* JournalFile) override;
+
+		struct FChunkInfoEntry
+		{
+			FIoChunkId ChunkId;
+			FCacheChunkInfo ChunkInfo;
+		};
+
+		FJournalChunkInfoPage& operator<<(const TArray<FChunkInfoEntry>& InEntries);
+
+	private:
+		bool bDirty = false;
+
+		TMap<FIoChunkId, FCacheChunkInfo> Entries;
+	};
+
+	class FJournalChunkPage : public FJournalPageBase
+	{
+	public:
+		FJournalChunkPage(int32 InMaxEntries);
+
+		EJournalPageResult Add(const FIoChunkId& ChunkId, const FCacheEntry& Entry);
+		bool GetEntry(const FIoChunkId& ChunkId, const uint64 ChunkOffset, const uint64 ChunkSize, FCacheEntry& OutEntry);
+		bool Contains(const FIoChunkId& ChunkId) const override
+		{
+			return ChunkMap.Contains(ChunkId);
+		}
+
+		void IterateChunkIds(TFunctionRef<void(const FIoChunkId& ChunkId)> Callback);
+		void IterateCacheEntriesForChunkId(const FIoChunkId& ChunkId, TFunctionRef<void(const FCacheEntry& Entry)> Callback);
+		void IterateCacheEntries(TFunctionRef<void(const FIoChunkId& ChunkId, const FCacheEntry& Entry)> Callback);
+
+		bool IsFull() const override
+		{
+			return Entries.Num() == MaxEntryCount;
+		}
+
+		bool Flush(IFileHandle* JournalFile) override;
+
+		int32 GetPageMaxSize() const override
+		{
+			return MaxEntryCount * Entries.GetTypeSize();
+		}
+
+		int32 GetEntrySize() const
+		{
+			return Entries.GetTypeSize();
+		}
+
+		void Invalidate(const FIoChunkId& ChunkId);
+		void InvalidateAll();
+
+		FJournalChunkPage& operator<<(TArray<FJournalStoreEntry>&& InEntries);
+
+	private:
+		using FIntervalMap = TMap<TInterval<uint64>, int32>;
+		TMap<FIoChunkId, FIntervalMap> ChunkMap;
+		TArray<FJournalStoreEntry> Entries;
+		bool bDirty = false;
+		bool bContainsInvalidEntries = false;
+	};
+
+	class FCacheJournalSectioned : public ICacheJournal
+	{
+	public:
+
+		FCacheJournalSectioned(FStringView JournalPath);
+		virtual ~FCacheJournalSectioned() override = default;
+
+		virtual void Flush() override;
+		virtual void InvalidateAll() override;
+		virtual void Invalidate(const FIoChunkId& ChunkId) override;
+
+		virtual bool SetChunkInfo(const FIoChunkId& ChunkId, const FModHash* OptModHash, const int64* OptRawSize, const int32* OptRawBlockSize) override;
+		virtual bool TryGetChunkInfo(const FIoChunkId& ChunkId, FCacheChunkInfo& OutChunkInfo) override;
+
+		virtual bool AddEntry(const FIoChunkId& ChunkId, const FCacheEntry& Entry) override;
+		virtual bool TryGetEntry(const FIoChunkId& ChunkId, const uint64 ChunkOffset, const uint64 ChunkSize, FCacheEntry& OutEntry) override;
+
+		virtual void IterateChunkIds(TFunctionRef<void(const FIoChunkId& ChunkId)> Callback) override;
+		virtual void IterateCacheEntriesForChunkId(const FIoChunkId& ChunkId, TFunctionRef<void(const FCacheEntry& Entry)> Callback) override;
+		virtual void IterateCacheEntries(TFunctionRef<void(const FIoChunkId& ChunkId, const FCacheEntry& Entry)> Callback) override;
+
+		virtual TArray<FIoChunkId> ApplyLatestServerChunkInfo(const TArray<FLatestServerChunkInfo>& ServerChunkInfoArray) override;
+
+	private:
+		void LoadCacheJournal(FStringView JournalPath);
+
+		FJournalChunkInfoPage* FindAddChunkInfoPage(const FIoChunkId& ChunkId);
+		FJournalChunkPage* FindChunkPage(const FIoChunkId& ChunkId);
+
+		int32 ReserveNextJournalFilePos(int32 PageSize)
+		{
+			int32 CurrentJournalPos = NextAvailableJournalPos;
+			NextAvailableJournalPos += sizeof(FJournalPageDescriptor) + PageSize;
+			return CurrentJournalPos;
+		}
+
+	private:
+		FString JournalFilePath;
+		uint32 JournalSectionCursor = sizeof(FJournalHeader);
+		TUniquePtr<IFileHandle> JournalFileHandle;
+		uint32 NextAvailableJournalPos = sizeof(FJournalHeader);
+
+		bool bPagesModified = false;
+
+		using FEntryIntervalMap = TMap<TInterval<uint64>, FCacheEntry>;
+
+		TArray<FJournalChunkPage> ChunkPages;
+		TArray<FJournalChunkInfoPage> ChunkInfoPages;
+		FJournalChunkPage* CurrentChunkPage = nullptr;
+
+		FCriticalSection CriticalSection;
+	};
+
+}
+
+#endif // !UE_BUILD_SHIPPING
\ No newline at end of file
