- Fixed compile errors when HAS_STORAGE_SERVER_RPC_GETCHUNKS_API is disabled

diff --git a/Engine/Source/Runtime/StorageServerClient/Private/StorageServerPlatformFile.cpp b/Engine/Source/Runtime/StorageServerClient/Private/StorageServerPlatformFile.cpp
--- a/Engine/Source/Runtime/StorageServerClient/Private/StorageServerPlatformFile.cpp
+++ b/Engine/Source/Runtime/StorageServerClient/Private/StorageServerPlatformFile.cpp
@@ -322,6 +322,7 @@ private:
 				break;
 			}
 
+#if HAS_STORAGE_SERVER_RPC_GETCHUNKS_API
 			PrecacheItems();
 			if (bExitWhenComplete)
 			{
@@ -331,6 +332,7 @@ private:
 				PrecacheItems();
 				break;
 			}
+#endif
 		}
 
 		IsCompleted->Trigger();
@@ -364,6 +366,7 @@ private:
 
 	}
 
+#if HAS_STORAGE_SERVER_RPC_GETCHUNKS_API
 	void PrecacheItems()
 	{
 		TRACE_CPUPROFILER_EVENT_SCOPE(FStorageServerEngineStartupPrecache::PrecacheItems);
@@ -403,6 +406,7 @@ private:
 			return;
 		}
 	}
+#endif
 
 	FStorageServerConnection& Connection;
 
