Correctly handle uncompressed chunks returned from Zen in StorageServerConnection.

[RN] fix

diff --git a/Engine/Source/Runtime/StorageServerClient/Private/StorageServerConnection.cpp b/Engine/Source/Runtime/StorageServerClient/Private/StorageServerConnection.cpp
--- a/Engine/Source/Runtime/StorageServerClient/Private/StorageServerConnection.cpp
+++ b/Engine/Source/Runtime/StorageServerClient/Private/StorageServerConnection.cpp
@@ -467,14 +467,15 @@ FIoStatus FStorageServerConnection::ReadChunkBatchRequest(const TArray<FChunkBat
 		const FCbObjectId Id = Chunk["Id"].AsObjectId();
 		const FIoHash RawHash = Chunk["RawHash"].AsHash();
 		const FIoHash FragmentHash = Chunk["FragmentHash"].AsHash();
+		const FIoHash Hash = Chunk["Hash"].AsHash();
 		const uint64 ModTag = Chunk["ModTag"].AsUInt64();
 
-		if (!ensureMsgf(RawHash != FIoHash() || FragmentHash != FIoHash(), TEXT("Failed to find hash in chunk info returned from server")))
+		if (!ensureMsgf(RawHash != FIoHash() || FragmentHash != FIoHash() || Hash != FIoHash(), TEXT("Failed to find hash in chunk info returned from server")))
 		{
 			continue;
 		}
 
-		const FIoHash& AttachmentHash = RawHash == FIoHash() ? FragmentHash : RawHash;
+		const FIoHash& AttachmentHash = (Hash != FIoHash() ? Hash : (RawHash != FIoHash() ? RawHash : FragmentHash));
 
 		int32* InfoIndex = AttachmentHashToIndex.Find(AttachmentHash);
 		if (!ensureMsgf(InfoIndex, TEXT("Failed to find hash in attachments returned from server")))
