Fix StorageServer failing to load files if the basedir was under projectdir or enginedir and the relative path to the file is a PORTION of the relative path to the projectdir or enginedir.  In these scenarios we have to incur the cost of allocating the absolute path, so minimize the allocations needed and the times when we attempt it.

diff --git a/Engine/Source/Runtime/StorageServerClient/Private/StorageServerPlatformFile.cpp b/Engine/Source/Runtime/StorageServerClient/Private/StorageServerPlatformFile.cpp
--- a/Engine/Source/Runtime/StorageServerClient/Private/StorageServerPlatformFile.cpp
+++ b/Engine/Source/Runtime/StorageServerClient/Private/StorageServerPlatformFile.cpp
@@ -1071,6 +1071,9 @@ bool FStorageServerPlatformFile::Initialize(IPlatformFile* Inner, const TCHAR* C
 
 void FStorageServerPlatformFile::InitializeAfterProjectFilePath()
 {
+	AbsEngineDir = FPaths::ConvertRelativePathToFull(FPlatformMisc::EngineDir());
+	AbsProjectDir = FPaths::ConvertRelativePathToFull(FPlatformMisc::ProjectDir());
+
 	InitializeConnection();
 
 	// optional debugging module depends on a valid Connection
@@ -1633,6 +1636,24 @@ bool FStorageServerPlatformFile::MakeStorageServerPath(const TCHAR* LocalFilenam
 		OutPath.Append(LocalFilenameOrDirectoryView.RightChop(LocalProjectDirView.Len()));
 		bValid = true;
 	}
+	else
+	{
+		TStringBuilder<128> AbsPathBuilder;
+		FPathViews::ToAbsolutePath(LocalFilenameOrDirectoryView, AbsPathBuilder);
+		FStringView RelativePath;
+		if (FPathViews::TryMakeChildPathRelativeTo(AbsPathBuilder, AbsProjectDir, RelativePath))
+		{
+			OutPath.Append(ServerProjectDirView);
+			OutPath.Append(RelativePath);
+			bValid = true;
+		}
+		else if (FPathViews::TryMakeChildPathRelativeTo(AbsPathBuilder, AbsEngineDir, RelativePath))
+		{
+			OutPath.Append(ServerEngineDirView);
+			OutPath.Append(RelativePath);
+			bValid = true;
+		}
+	}
 
 	if (bValid)
 	{
diff --git a/Engine/Source/Runtime/StorageServerClient/Private/StorageServerPlatformFile.h b/Engine/Source/Runtime/StorageServerClient/Private/StorageServerPlatformFile.h
--- a/Engine/Source/Runtime/StorageServerClient/Private/StorageServerPlatformFile.h
+++ b/Engine/Source/Runtime/StorageServerClient/Private/StorageServerPlatformFile.h
@@ -159,6 +159,8 @@ private:
 	FString ServerProject;
 	FString ServerPlatform;
 	FString BaseURI;
+	FString AbsProjectDir;
+	FString AbsEngineDir;
 	TDirectoryTree<FString> RemapDirectoriesTree;
 	mutable TArray<FString> HostAddrs;
 	mutable uint16 HostPort = 8558;
