Verifying storage server client cache entries validity via ModTag

diff --git a/Engine/Source/Runtime/StorageServerClient/Private/Cache/CacheJournal.h b/Engine/Source/Runtime/StorageServerClient/Private/Cache/CacheJournal.h
--- a/Engine/Source/Runtime/StorageServerClient/Private/Cache/CacheJournal.h
+++ b/Engine/Source/Runtime/StorageServerClient/Private/Cache/CacheJournal.h
@@ -13,34 +13,23 @@
 
 namespace StorageServer
 {
-	// Opaque identifier that changes when content of chunk changes,
-	// right now is implemented via ChunkId RawHash as a stop gap solution
-	typedef FIoHash FModHash;
-
-	struct FLatestServerChunkInfo
-	{
-		FIoChunkId ChunkId;
-		FModHash ModHash;
-		int64 RawSize;
-	};
-
 	struct FCacheChunkInfo
 	{
-		FModHash ModHash;   // content modification chunk hash
-		int64 RawSize;      // size of complete chunk
-		int32 RawBlockSize; // size of block in a chunk if any, e.g. 256kb, can be 0
+		TOptional<uint64> ModTag;      // content modification tag
+		TOptional<int64> RawSize;      // size of complete chunk
+		TOptional<int32> RawBlockSize; // size of block in a chunk if any, e.g. 256kb, can be 0
 
 		bool SetChunkInfo(
-			const FModHash* OptModHash,
-			const int64* OptRawSize,
-			const int32* OptRawBlockSize
+			const TOptional<uint64>& OptModTag,
+			const TOptional<int64>& OptRawSize,
+			const TOptional<int32>& OptRawBlockSize
 		)
 		{
 			bool bResult = true;
-			if (OptModHash)
+			if (OptModTag)
 			{
-				bResult = ModHash == *OptModHash;
-				ModHash = *OptModHash;
+				bResult = ModTag.IsSet() ? *ModTag == *OptModTag : true;
+				ModTag = *OptModTag;
 			}
 			if (OptRawSize)
 			{
@@ -53,9 +42,17 @@ namespace StorageServer
 			return bResult;
 		}
 
+		FCacheChunkInfo& operator=(const FCacheChunkInfo& Other)
+		{
+			ModTag = Other.ModTag;
+			RawSize = Other.RawSize;
+			RawBlockSize = Other.RawBlockSize;
+			return *this;
+		}
+
 		friend FArchive& operator << (FArchive& Ar, FCacheChunkInfo& Info)
 		{
-			Ar << Info.ModHash;
+			Ar << Info.ModTag;
 			Ar << Info.RawSize;
 			Ar << Info.RawBlockSize;
 			return Ar;
@@ -105,15 +102,15 @@ namespace StorageServer
 		virtual void Invalidate(const FIoChunkId& ChunkId) = 0;
 
 		// Updates chunk info. All fields are optional.
-		// If OptModHash is passed will compare it to existing ModHash and invalidate all data for this chunk id if they don't match.
+		// If OptModTag is passed will compare it to existing ModTag and asks you to invalidate all data for this chunk id if they don't match.
 		// OptRawSize, OptRawBlockSize will simply be updated if passed.
-		// Returns true if either new ChunkId info is created or ModHash matches previous existing entry for this ChunkId.
+		// Returns true if either new ChunkId info is created or ModTag matches previous existing entry for this ChunkId.
 		// Returns false if ModHash doesn't match and cached data for this entry needs to be invalidated.
 		virtual bool SetChunkInfo(
 			const FIoChunkId& ChunkId,
-			const FModHash* OptModHash,
-			const int64* OptRawSize,
-			const int32* OptRawBlockSize
+			const TOptional<uint64>& OptModTag,
+			const TOptional<int64>& OptRawSize,
+			const TOptional<int32>& OptRawBlockSize
 		) = 0;
 
 		// Return chunk info if present, returns true if chunk info is present.
@@ -140,7 +137,7 @@ namespace StorageServer
 		// Iterate chunk id in journal that contain at least one entry.
 		// Not safe to invoke other journal methods from the callback.
 		virtual void IterateChunkIds(
-			TFunctionRef<void(const FIoChunkId& ChunkId)> Callback
+			TFunctionRef<void(const FIoChunkId& ChunkId, const FCacheChunkInfo& ChunkInfo)> Callback
 		) = 0;
 
 		// Iterate entries for a chunk id in journal.
@@ -155,12 +152,6 @@ namespace StorageServer
 		virtual void IterateCacheEntries(
 			TFunctionRef<void(const FIoChunkId& ChunkId, const FCacheEntry& Entry)> Callback
 		) = 0;
-
-		// Updates a group of chunk infos based on the latest from the server. Expected to be used during initialization
-		// Returns all invalidated chunk Ids
-		virtual TArray<FIoChunkId> ApplyLatestServerChunkInfo( 
-			const TArray<FLatestServerChunkInfo>& ServerChunkInfoArray 
-		) = 0;
 	};
 }
 
diff --git a/Engine/Source/Runtime/StorageServerClient/Private/Cache/CacheJournalSectioned.cpp b/Engine/Source/Runtime/StorageServerClient/Private/Cache/CacheJournalSectioned.cpp
--- a/Engine/Source/Runtime/StorageServerClient/Private/Cache/CacheJournalSectioned.cpp
+++ b/Engine/Source/Runtime/StorageServerClient/Private/Cache/CacheJournalSectioned.cpp
@@ -106,21 +106,27 @@ namespace StorageServer
 		}
 	}
 
-	bool FCacheJournalSectioned::SetChunkInfo(const FIoChunkId& ChunkId, const FModHash* OptModHash, const int64* OptRawSize, const int32* OptRawBlockSize)
+	bool FCacheJournalSectioned::SetChunkInfo(const FIoChunkId& ChunkId, const TOptional<uint64>& OptModHash, const TOptional<int64>& OptRawSize, const TOptional<int32>& OptRawBlockSize)
 	{
 		UE::TScopeLock _(CriticalSection);
 
 		FCacheChunkInfo ChunkInfo;
-		if (TryGetChunkInfo(ChunkId, ChunkInfo))
+		for (FJournalChunkInfoPage& Page : ChunkInfoPages)
 		{
-			return ChunkInfo.SetChunkInfo(OptModHash, OptRawSize, OptRawBlockSize);
+			if (Page.GetEntry(ChunkId, ChunkInfo))
+			{
+				const bool bHashesMatch = ChunkInfo.SetChunkInfo(OptModHash, OptRawSize, OptRawBlockSize);
+				Page.Add(ChunkId, ChunkInfo);
+				return bHashesMatch;
+			}
 		}
 
 		FJournalChunkInfoPage* Page = FindAddChunkInfoPage(ChunkId);
 		check(Page != nullptr && !Page->IsFull());
 
 		ChunkInfo.SetChunkInfo(OptModHash, OptRawSize, OptRawBlockSize);
-		return Page->Add(ChunkId, ChunkInfo) == EJournalPageResult::Ok;
+		Page->Add(ChunkId, ChunkInfo);
+		return true;
 	}
 
 	bool FCacheJournalSectioned::TryGetChunkInfo(const FIoChunkId& ChunkId, FCacheChunkInfo& OutChunkInfo)
@@ -172,11 +178,11 @@ namespace StorageServer
 		return false;
 	}
 
-	void FCacheJournalSectioned::IterateChunkIds(TFunctionRef<void(const FIoChunkId& ChunkId)> Callback)
+	void FCacheJournalSectioned::IterateChunkIds(TFunctionRef<void(const FIoChunkId& ChunkId, const FCacheChunkInfo& ChunkInfo)> Callback)
 	{
 		UE::TScopeLock _(CriticalSection);
 
-		for (FJournalChunkPage& Page : ChunkPages)
+		for (FJournalChunkInfoPage& Page : ChunkInfoPages)
 		{
 			Page.IterateChunkIds(MoveTemp(Callback));
 		}
@@ -202,31 +208,6 @@ namespace StorageServer
 		}
 	}
 
-	TArray<FIoChunkId> FCacheJournalSectioned::ApplyLatestServerChunkInfo(const TArray<FLatestServerChunkInfo>& ServerChunkInfoArray)
-	{
-		UE::TScopeLock _(CriticalSection);
-
-		TArray<FIoChunkId> InvalidatedChunks;
-		for (const FLatestServerChunkInfo& ServerChunkInfo : ServerChunkInfoArray)
-		{
-			FCacheChunkInfo ChunkInfo;
-			if (TryGetChunkInfo(ServerChunkInfo.ChunkId, ChunkInfo))
-			{
-				const bool bHashesMatch = ChunkInfo.SetChunkInfo(&ServerChunkInfo.ModHash, &ServerChunkInfo.RawSize, nullptr);
-				if (!bHashesMatch)
-				{
-					InvalidatedChunks.Add(ServerChunkInfo.ChunkId);
-				}
-			}
-			else
-			{
-				SetChunkInfo(ServerChunkInfo.ChunkId, &ServerChunkInfo.ModHash, &ServerChunkInfo.RawSize, nullptr);
-			}
-		}
-
-		return InvalidatedChunks;
-	}
-
 	void FCacheJournalSectioned::LoadCacheJournal(FStringView JournalPath)
 	{
 		TRACE_CPUPROFILER_EVENT_SCOPE(FCacheJournalSectioned::LoadCacheJournal);
@@ -415,26 +396,20 @@ namespace StorageServer
 
 	EJournalPageResult FJournalChunkInfoPage::Add(const FIoChunkId& ChunkId, const FCacheChunkInfo& Entry)
 	{
-		if (Entries.Num() >= MaxEntryCount)
-		{
-			return EJournalPageResult::PageFull;
-		}
-
 		FCacheChunkInfo* ChunkInfo = Entries.Find(ChunkId);
 		if (ChunkInfo)
 		{
-			FModHash PreviousHash = ChunkInfo->ModHash;
-			const bool bHashesMatch = ChunkInfo->SetChunkInfo(&Entry.ModHash, &Entry.RawSize, &Entry.RawBlockSize);
-			if (!bHashesMatch)
-			{
-				bDirty = true;
-			}
+			*ChunkInfo = Entry;
 			return EJournalPageResult::EntryAlreadyExists;
 		}
 
+		if (Entries.Num() >= MaxEntryCount)
+		{
+			return EJournalPageResult::PageFull;
+		}
+
 		Entries.Add(ChunkId, Entry);
 		bDirty = true;
-
 		return EJournalPageResult::Ok;
 	}
 
@@ -466,6 +441,14 @@ namespace StorageServer
 		Entries.Empty();
 	}
 
+	void FJournalChunkInfoPage::IterateChunkIds(TFunctionRef<void(const FIoChunkId& ChunkId, const FCacheChunkInfo& ChunkInfo)> Callback)
+	{
+		for (const TTuple<FIoChunkId, FCacheChunkInfo>& Entry : Entries)
+		{
+			Callback(Entry.Key, Entry.Value);
+		}
+	}
+
 	bool FJournalChunkInfoPage::Flush(IFileHandle* Handle)
 	{
 		TRACE_CPUPROFILER_EVENT_SCOPE(FJournalChunkInfoPage::Flush);
@@ -566,14 +549,6 @@ namespace StorageServer
 		return false;
 	}
 
-	void FJournalChunkPage::IterateChunkIds(TFunctionRef<void(const FIoChunkId& ChunkId)> Callback)
-	{
-		for (const FJournalStoreEntry& Entry : Entries)
-		{
-			Callback(Entry.ChunkId);
-		}
-	}
-
 	void FJournalChunkPage::IterateCacheEntriesForChunkId(const FIoChunkId& ChunkId, TFunctionRef<void(const FCacheEntry& Entry)> Callback)
 	{
 		FIntervalMap* IntervalMap = ChunkMap.Find(ChunkId);
diff --git a/Engine/Source/Runtime/StorageServerClient/Private/Cache/CacheJournalSectioned.h b/Engine/Source/Runtime/StorageServerClient/Private/Cache/CacheJournalSectioned.h
--- a/Engine/Source/Runtime/StorageServerClient/Private/Cache/CacheJournalSectioned.h
+++ b/Engine/Source/Runtime/StorageServerClient/Private/Cache/CacheJournalSectioned.h
@@ -107,12 +107,14 @@ namespace StorageServer
 		bool Contains(const FIoChunkId& ChunkId) const override;
 		bool IsFull() const override 
 		{
-			return Entries.Num() == MaxEntryCount;
+			return Entries.Num() >= MaxEntryCount;
 		}
 
 		void Invalidate(const FIoChunkId& ChunkId);
 		void InvalidateAll();
 
+		void IterateChunkIds(TFunctionRef<void(const FIoChunkId& ChunkId, const FCacheChunkInfo& ChunkInfo)> Callback);
+
 		int32 GetPageMaxSize() const override
 		{
 			return MaxEntryCount * sizeof(FChunkInfoEntry);
@@ -146,7 +148,6 @@ namespace StorageServer
 			return ChunkMap.Contains(ChunkId);
 		}
 
-		void IterateChunkIds(TFunctionRef<void(const FIoChunkId& ChunkId)> Callback);
 		void IterateCacheEntriesForChunkId(const FIoChunkId& ChunkId, TFunctionRef<void(const FCacheEntry& Entry)> Callback);
 		void IterateCacheEntries(TFunctionRef<void(const FIoChunkId& ChunkId, const FCacheEntry& Entry)> Callback);
 
@@ -191,18 +192,16 @@ namespace StorageServer
 		virtual void InvalidateAll() override;
 		virtual void Invalidate(const FIoChunkId& ChunkId) override;
 
-		virtual bool SetChunkInfo(const FIoChunkId& ChunkId, const FModHash* OptModHash, const int64* OptRawSize, const int32* OptRawBlockSize) override;
+		virtual bool SetChunkInfo(const FIoChunkId& ChunkId, const TOptional<uint64>& OptModHash, const TOptional<int64>& OptRawSize, const TOptional<int32>& OptRawBlockSize) override;
 		virtual bool TryGetChunkInfo(const FIoChunkId& ChunkId, FCacheChunkInfo& OutChunkInfo) override;
 
 		virtual bool AddEntry(const FIoChunkId& ChunkId, const FCacheEntry& Entry) override;
 		virtual bool TryGetEntry(const FIoChunkId& ChunkId, const uint64 ChunkOffset, const uint64 ChunkSize, FCacheEntry& OutEntry) override;
 
-		virtual void IterateChunkIds(TFunctionRef<void(const FIoChunkId& ChunkId)> Callback) override;
+		virtual void IterateChunkIds(TFunctionRef<void(const FIoChunkId& ChunkId, const FCacheChunkInfo& ChunkInfo)> Callback) override;
 		virtual void IterateCacheEntriesForChunkId(const FIoChunkId& ChunkId, TFunctionRef<void(const FCacheEntry& Entry)> Callback) override;
 		virtual void IterateCacheEntries(TFunctionRef<void(const FIoChunkId& ChunkId, const FCacheEntry& Entry)> Callback) override;
 
-		virtual TArray<FIoChunkId> ApplyLatestServerChunkInfo(const TArray<FLatestServerChunkInfo>& ServerChunkInfoArray) override;
-
 	private:
 		void LoadCacheJournal(FStringView JournalPath);
 
diff --git a/Engine/Source/Runtime/StorageServerClient/Private/Cache/CacheJournalSimple.cpp b/Engine/Source/Runtime/StorageServerClient/Private/Cache/CacheJournalSimple.cpp
--- a/Engine/Source/Runtime/StorageServerClient/Private/Cache/CacheJournalSimple.cpp
+++ b/Engine/Source/Runtime/StorageServerClient/Private/Cache/CacheJournalSimple.cpp
@@ -114,6 +114,8 @@ namespace StorageServer
 	{
 		FScopeLock Lock(&DataLock);
 
+		TRACE_COUNTER_INCREMENT(ZenJournalSimpleInvalidated);
+
 		FPerChunkCacheEntries* CacheEntries = ChunkCacheEntries.Find(ChunkId);
 		if (CacheEntries)
 		{
@@ -126,59 +128,22 @@ namespace StorageServer
 		bDirty = true;
 	}
 
-	bool FCacheJournalSimple::SetChunkInfo(const FIoChunkId& ChunkId, const FIoHash* OptModHash, const int64* OptRawSize, const int32* OptRawBlockSize)
+	bool FCacheJournalSimple::SetChunkInfo(const FIoChunkId& ChunkId, const TOptional<uint64>& OptModTag, const TOptional<int64>& OptRawSize, const TOptional<int32>& OptRawBlockSize)
 	{
 		FScopeLock Lock(&DataLock);
 		FCacheChunkInfo* ChunkInfo = ChunkInfos.Find(ChunkId);
 		if (ChunkInfo)
 		{
-			FModHash PreviousHash = ChunkInfo->ModHash;
-			const bool bHashesMatch = ChunkInfo->SetChunkInfo(OptModHash, OptRawSize, OptRawBlockSize);
-			if (!bHashesMatch)
-			{
-				TRACE_COUNTER_INCREMENT(ZenJournalSimpleInvalidated);
-			}
-			return bHashesMatch;
+			return ChunkInfo->SetChunkInfo(OptModTag, OptRawSize, OptRawBlockSize);
 		}
 
 		FCacheChunkInfo NewChunkInfo = {};
-		NewChunkInfo.SetChunkInfo(OptModHash, OptRawSize, OptRawBlockSize);
+		NewChunkInfo.SetChunkInfo(OptModTag, OptRawSize, OptRawBlockSize);
 		ChunkInfos.Add(ChunkId, NewChunkInfo);
 		bDirty = true;
 		return true;
 	}
 
-	TArray<FIoChunkId> FCacheJournalSimple::ApplyLatestServerChunkInfo(const TArray<FLatestServerChunkInfo>& ServerChunkInfoArray)
-	{
-		TRACE_CPUPROFILER_EVENT_SCOPE(FCacheJournalSimple::ApplyLatestServerChunkInfo);
-
-		TArray<FIoChunkId> InvalidatedChunks;
-		for (const FLatestServerChunkInfo& ServerChunkInfo : ServerChunkInfoArray)
-		{
-			FScopeLock Lock(&DataLock);
-			FCacheChunkInfo* ChunkInfo = ChunkInfos.Find(ServerChunkInfo.ChunkId);
-			if (ChunkInfo)
-			{
-				FModHash PreviousHash = ChunkInfo->ModHash;
-				const bool bHashesMatch = ChunkInfo->SetChunkInfo(&ServerChunkInfo.ModHash, &ServerChunkInfo.RawSize, nullptr);
-				if (!bHashesMatch)
-				{
-					TRACE_COUNTER_INCREMENT(ZenJournalSimpleInvalidated);
-					InvalidatedChunks.Add(ServerChunkInfo.ChunkId);
-				}
-			}
-			else
-			{
-				FCacheChunkInfo NewChunkInfo = {};
-				NewChunkInfo.SetChunkInfo(&ServerChunkInfo.ModHash, &ServerChunkInfo.RawSize, nullptr);
-				ChunkInfos.Add(ServerChunkInfo.ChunkId, NewChunkInfo);
-				bDirty = true;
-			}
-		}
-
-		return MoveTemp(InvalidatedChunks);
-	}
-
 	bool FCacheJournalSimple::TryGetChunkInfo(const FIoChunkId& ChunkId, FCacheChunkInfo& OutChunkInfo)
 	{
 		FScopeLock Lock(&DataLock);
@@ -252,13 +217,16 @@ namespace StorageServer
 		return true;
 	}
 
-	void FCacheJournalSimple::IterateChunkIds(TFunctionRef<void(const FIoChunkId& ChunkId)> Callback)
+	void FCacheJournalSimple::IterateChunkIds(TFunctionRef<void(const FIoChunkId& ChunkId, const FCacheChunkInfo& ChunkInfo)> Callback)
 	{
 		FScopeLock Lock(&DataLock);
 
+		const FCacheChunkInfo EmptyChunkInfo = {};
+
 		for (const TTuple<FIoChunkId, FPerChunkCacheEntries>& Pair : ChunkCacheEntries)
 		{
-			Callback(Pair.Key);
+			const FCacheChunkInfo* ChunkInfo = ChunkInfos.Find(Pair.Key);
+			Callback(Pair.Key, ChunkInfo ? *ChunkInfo : EmptyChunkInfo);
 		}
 	}
 
diff --git a/Engine/Source/Runtime/StorageServerClient/Private/Cache/CacheJournalSimple.h b/Engine/Source/Runtime/StorageServerClient/Private/Cache/CacheJournalSimple.h
--- a/Engine/Source/Runtime/StorageServerClient/Private/Cache/CacheJournalSimple.h
+++ b/Engine/Source/Runtime/StorageServerClient/Private/Cache/CacheJournalSimple.h
@@ -26,9 +26,9 @@ namespace StorageServer
 
 		virtual bool SetChunkInfo(
 			const FIoChunkId& ChunkId,
-			const FIoHash* OptModHash,
-			const int64* OptRawSize,
-			const int32* OptRawBlockSize
+			const TOptional<uint64>& OptModTag,
+			const TOptional<int64>& OptRawSize,
+			const TOptional<int32>& OptRawBlockSize
 		) override;
 		virtual bool TryGetChunkInfo(
 			const FIoChunkId& ChunkId,
@@ -47,7 +47,7 @@ namespace StorageServer
 		) override;
 
 		virtual void IterateChunkIds(
-			TFunctionRef<void(const FIoChunkId& ChunkId)> Callback
+			TFunctionRef<void(const FIoChunkId& ChunkId, const FCacheChunkInfo& ChunkInfo)> Callback
 		) override;
 		virtual void IterateCacheEntriesForChunkId(
 			const FIoChunkId& ChunkId,
@@ -57,10 +57,6 @@ namespace StorageServer
 			TFunctionRef<void(const FIoChunkId& ChunkId, const FCacheEntry& Entry)> Callback
 		) override;
 
-		virtual TArray<FIoChunkId> ApplyLatestServerChunkInfo(
-			const TArray<FLatestServerChunkInfo>& ServerChunkInfoArray 
-		) override;
-		
 	private:
 		FCriticalSection DataLock, FlushTaskLock;
 		FString FileName;
diff --git a/Engine/Source/Runtime/StorageServerClient/Private/Cache/CacheStrategy.h b/Engine/Source/Runtime/StorageServerClient/Private/Cache/CacheStrategy.h
--- a/Engine/Source/Runtime/StorageServerClient/Private/Cache/CacheStrategy.h
+++ b/Engine/Source/Runtime/StorageServerClient/Private/Cache/CacheStrategy.h
@@ -42,22 +42,9 @@ namespace StorageServer
 
 		virtual void Invalidate(const FIoChunkId& ChunkId) = 0;
 
-		// Stop gap implementation till new endpoint is available in zen
-		virtual void SetChunkInfo(
-			const FIoChunkId& ChunkId,
-			const FModHash& ModHash,
-			const int64& RawSize
-		)
-		{
-			if (!Journal->SetChunkInfo(ChunkId, &ModHash, &RawSize, nullptr))
-			{
-				Invalidate(ChunkId);
-			}
-		}
-
 		virtual void CacheChunkSize(const FIoChunkId& ChunkId, const int64 RawSize)
 		{
-			Journal->SetChunkInfo(ChunkId, nullptr, &RawSize, nullptr);
+			Journal->SetChunkInfo(ChunkId, TOptional<uint64>(), RawSize, TOptional<int32>());
 		}
 
 		virtual bool TryGetChunkSize(const FIoChunkId& ChunkId, int64& OutRawSize)
@@ -65,8 +52,15 @@ namespace StorageServer
 			FCacheChunkInfo ChunkInfo;
 			if (Journal->TryGetChunkInfo(ChunkId, ChunkInfo))
 			{
-				OutRawSize = ChunkInfo.RawSize;
-				return true;
+				if (ChunkInfo.RawSize.IsSet())
+				{
+					OutRawSize = *ChunkInfo.RawSize;
+					return true;
+				}
+				else
+				{
+					return false;
+				}
 			}
 			else
 			{
@@ -94,7 +88,8 @@ namespace StorageServer
 			const uint64 RequestRawOffset,
 			const uint64 RequestRawSize,
 			const FIoBuffer& ResultBuffer,
-			const EStorageServerContentType ResultContentType
+			const EStorageServerContentType ResultContentType,
+			const uint64 ResultModTag
 		) = 0;
 
 		virtual void Flush()
@@ -103,14 +98,9 @@ namespace StorageServer
 			Storage->Flush();
 		}
 
-		virtual void ApplyLatestServerChunkInfo(const TArray<FLatestServerChunkInfo>& ServerChunkInfoArray)
+		virtual void IterateChunkIds(TFunctionRef<void(const FIoChunkId& ChunkId, const FCacheChunkInfo& ChunkInfo)> Callback)
 		{
-			// naive implementation - specializations can likely do better
-			TArray<FIoChunkId> InvalidatedChunks = Journal->ApplyLatestServerChunkInfo(ServerChunkInfoArray);
-			for (const FIoChunkId& ChunkId : InvalidatedChunks)
-			{
-				Invalidate(ChunkId);
-			}
+			Journal->IterateChunkIds(Callback);
 		}
 
 	protected:
diff --git a/Engine/Source/Runtime/StorageServerClient/Private/Cache/CacheStrategyLinear.cpp b/Engine/Source/Runtime/StorageServerClient/Private/Cache/CacheStrategyLinear.cpp
--- a/Engine/Source/Runtime/StorageServerClient/Private/Cache/CacheStrategyLinear.cpp
+++ b/Engine/Source/Runtime/StorageServerClient/Private/Cache/CacheStrategyLinear.cpp
@@ -57,19 +57,31 @@ namespace StorageServer
 		UE_LOG(LogCacheStrategyLinear, Display, TEXT("Abandoning zen linear cache"));
 #endif
 
+		FScopeLock ScopeLock(&JournalLock);
+
 		Journal->InvalidateAll();
 		CurrentSize = 0;
 		CurrentInvalidSize = 0;
 		SetCounters();
 	}
 
-	void FCacheStrategyLinear::Invalidate(const FIoChunkId& ChunkId)
+	void FCacheStrategyLinear::Invalidate(const FIoChunkId& ChunkId, bool bShouldLockJournal)
 	{
+		UE::TConditionalScopeLock Lock(JournalLock, bShouldLockJournal);
+
 		if (CurrentSize == 0)
 		{
 			return;
 		}
 
+#if ZEN_CACHE_VERBOSE_LOG
+		FCacheChunkInfo ChunkInfo = {};
+		Journal->TryGetChunkInfo(ChunkId, ChunkInfo);
+		UE_LOG(LogCacheStrategyLinear, Display, TEXT("CachePut Invalidate2 %s was %llx"),
+			*LexToString(ChunkId),
+			ChunkInfo.ModTag.Get(0));
+#endif
+
 		Journal->IterateCacheEntriesForChunkId(ChunkId, [this](const FCacheEntry& CacheEntry)
 		{
 			CurrentInvalidSize += CacheEntry.StorageSize;
@@ -81,34 +93,9 @@ namespace StorageServer
 		{
 			InvalidateAll();
 		}
-	}
-
-	void FCacheStrategyLinear::ApplyLatestServerChunkInfo(const TArray<FLatestServerChunkInfo>& ServerChunkInfoArray)
-	{
-		// update the journal
-		TArray<FIoChunkId> InvalidatedChunks = Journal->ApplyLatestServerChunkInfo(ServerChunkInfoArray);
-
-		// invalidate affected items
-		if (CurrentSize > 0)
+		else
 		{
-			TRACE_CPUPROFILER_EVENT_SCOPE(FCacheStrategyLinear::Invalidate);
-
-			for (const FIoChunkId& ChunkId : InvalidatedChunks)
-			{
-				Journal->IterateCacheEntriesForChunkId(ChunkId, [this](const FCacheEntry& CacheEntry)
-				{
-					CurrentInvalidSize += CacheEntry.StorageSize;
-				});
-
-				if (CurrentInvalidSize >= AbandonAtInvalidSize)
-				{
-					SetCounters();
-					InvalidateAll();
-					return;
-				}
-			}
-
-			SetCounters();
+			Journal->Invalidate(ChunkId);
 		}
 	}
 
@@ -172,7 +159,7 @@ namespace StorageServer
 		return true;
 	}
 
-	void FCacheStrategyLinear::CacheChunk(const FIoChunkId& RequestChunkId, const uint64 RequestOffset, const uint64 RequestSize, const FIoBuffer& Buffer, const EStorageServerContentType ContentType)
+	void FCacheStrategyLinear::CacheChunk(const FIoChunkId& RequestChunkId, const uint64 RequestOffset, const uint64 RequestSize, const FIoBuffer& Buffer, const EStorageServerContentType ContentType, const uint64 ResultModTag)
 	{
 		// TODO add hint to ignore some requests to limit overall write speed
 
@@ -181,6 +168,30 @@ namespace StorageServer
 
 		{
 			FScopeLock ScopeLock(&JournalLock);
+#if ZEN_CACHE_VERBOSE_LOG
+			FCacheChunkInfo ChunkInfoBefore = {};
+			Journal->TryGetChunkInfo(RequestChunkId, ChunkInfoBefore);
+#endif
+
+			if (!UNLIKELY(Journal->SetChunkInfo(RequestChunkId, ResultModTag, TOptional<int64>(), TOptional<int32>())))
+			{
+				// Invalidate whole chunk in unlikely case if ModTag's don't match, for example if asset was changed in runtime.
+#if ZEN_CACHE_VERBOSE_LOG
+				UE_LOG(LogCacheStrategyLinear, Display, TEXT("CachePut Invalidate %s %llu:%llu was %llx become %llx"),
+					*LexToString(RequestChunkId),
+					RequestOffset,
+					RequestSize,
+					ChunkInfoBefore.ModTag.Get(0),
+					ResultModTag
+				);
+#endif
+
+				Invalidate(RequestChunkId, false);
+
+				// Still need to set chunk info after invalidation
+				const bool ModTagValid = Journal->SetChunkInfo(RequestChunkId, ResultModTag, TOptional<int64>(), TOptional<int32>());
+				ensure(ModTagValid);
+			}
 
 			if (CurrentSize.load() + Buffer.GetSize() >= Storage->GetSize())
 			{
@@ -202,13 +213,14 @@ namespace StorageServer
 				SetCounters();
 
 #if ZEN_CACHE_VERBOSE_LOG
-				UE_LOG(LogCacheStrategyLinear, Display, TEXT("CachePut %s %llu:%llu->%llu:%llu %s"),
+				UE_LOG(LogCacheStrategyLinear, Display, TEXT("CachePut %s %llu:%llu->%llu:%llu %s:%llx"),
 					*LexToString(RequestChunkId),
 					RequestOffset,
 					RequestSize,
 					Entry.StorageOffset,
 					Entry.StorageSize,
-					*LexToString(Entry.StorageHash)
+					*LexToString(Entry.StorageHash),
+					ResultModTag
 				);
 #endif
 			}
diff --git a/Engine/Source/Runtime/StorageServerClient/Private/Cache/CacheStrategyLinear.h b/Engine/Source/Runtime/StorageServerClient/Private/Cache/CacheStrategyLinear.h
--- a/Engine/Source/Runtime/StorageServerClient/Private/Cache/CacheStrategyLinear.h
+++ b/Engine/Source/Runtime/StorageServerClient/Private/Cache/CacheStrategyLinear.h
@@ -17,14 +17,15 @@ namespace StorageServer
 		virtual ~FCacheStrategyLinear() override;
 
 		void InvalidateAll();
-		virtual void Invalidate(const FIoChunkId& ChunkId) override;
+		virtual void Invalidate(const FIoChunkId& ChunkId) override
+		{
+			Invalidate(ChunkId, true);
+		}
 
 		virtual bool ReadChunk(const FIoChunkId& RequestChunkId, const uint64 RequestOffset, const uint64 RequestSize, TOptional<FIoBuffer> OptDestination, FIoBuffer& OutBuffer, EStorageServerContentType& ContentType) override;
 
-		virtual void CacheChunk(const FIoChunkId& RequestChunkId, const uint64 RequestOffset, const uint64 RequestSize, const FIoBuffer& Buffer, const EStorageServerContentType ContentType) override;
+		virtual void CacheChunk(const FIoChunkId& RequestChunkId, const uint64 RequestOffset, const uint64 RequestSize, const FIoBuffer& Buffer, const EStorageServerContentType ContentType, const uint64 ResultModTag) override;
 
-		virtual void ApplyLatestServerChunkInfo(const TArray<FLatestServerChunkInfo>& ServerChunkInfoArray) override;
-		
 	private:
 		// Linear cache is implemented as arena allocator
 		std::atomic<uint64> CurrentSize; // Current size of arena allocator
@@ -32,6 +33,7 @@ namespace StorageServer
 		uint64 AbandonAtInvalidSize; // Consider cache to be a lost cause at this amount of stale data
 		FCriticalSection JournalLock;
 
+		void Invalidate(const FIoChunkId& ChunkId, bool bShouldLockJournal);
 		void SetCounters();
 	};
 }
diff --git a/Engine/Source/Runtime/StorageServerClient/Private/StorageServerConnection.cpp b/Engine/Source/Runtime/StorageServerClient/Private/StorageServerConnection.cpp
--- a/Engine/Source/Runtime/StorageServerClient/Private/StorageServerConnection.cpp
+++ b/Engine/Source/Runtime/StorageServerClient/Private/StorageServerConnection.cpp
@@ -271,6 +271,7 @@ TIoStatusOr<FIoBuffer> FStorageServerConnection::ReadChunkRequest(
 	const double StartTime = FPlatformTime::Seconds();
 
 	IStorageServerHttpClient::FResult ResultTuple;
+	uint64 ResultModTag = 0;
 
 	bool bWasCached = false;
 	FIoBuffer CacheChunkBuffer;
@@ -286,13 +287,14 @@ TIoStatusOr<FIoBuffer> FStorageServerConnection::ReadChunkRequest(
 	{
 #if HAS_STORAGE_SERVER_RPC_GETCHUNKS_API
 		TArray<FChunkBatchRequestEntry> ChunkBatchRequests;
-		TOptional<uint64> EmptyModTag;
-		ChunkBatchRequests.Emplace(FChunkBatchRequestEntry{ChunkId, Offset, Size, EmptyModTag});
+		ChunkBatchRequests.Add(FChunkBatchRequestEntry::DataRequest(ChunkId, Offset, Size));
 
-		FIoStatus ResultStatus = ReadChunkBatchRequest(ChunkBatchRequests, [&](FIoChunkId Id, EStorageServerContentType MimeType, FIoBuffer Data, uint64 ModTag)
+		FIoStatus ResultStatus = ReadChunkBatchRequest(ChunkBatchRequests, [&](FIoChunkId Id, EStorageServerContentType MimeType, FIoBuffer Data, const TOptional<uint64>& ModTag)
 		{
+			ensure(ChunkId == Id);
 			Data.MakeOwned();
 			ResultTuple = IStorageServerHttpClient::FResult(Data, MimeType);
+			ResultModTag = ModTag.GetValue(); // ModTag must be present for cache invalidation to work
 		});
 
 		if (!ResultStatus.IsOk())
@@ -306,6 +308,7 @@ TIoStatusOr<FIoBuffer> FStorageServerConnection::ReadChunkRequest(
 		BuildReadChunkRequestUrl(ResourceBuilder, ChunkId, Offset, Size);
 
 		ResultTuple = HttpClient->RequestSync(*ResourceBuilder);
+		ResultModTag = 0; // This endpoint doesn't support ModTag.
 #endif
 	}
 
@@ -316,7 +319,7 @@ TIoStatusOr<FIoBuffer> FStorageServerConnection::ReadChunkRequest(
 		if (!bWasCached && CacheStrategy.IsValid())
 		{
 			ZenCacheMiss(ResultTuple.Key.ValueOrDie().GetSize());
-			CacheStrategy->CacheChunk(ChunkId, Offset, Size, ResultTuple.Key.ValueOrDie(), ResultTuple.Value);
+			CacheStrategy->CacheChunk(ChunkId, Offset, Size, ResultTuple.Key.ValueOrDie(), ResultTuple.Value, ResultModTag);
 		}
 
 		const double Duration = FPlatformTime::Seconds() - StartTime;
@@ -326,8 +329,21 @@ TIoStatusOr<FIoBuffer> FStorageServerConnection::ReadChunkRequest(
 	return ResultBuffer;
 }
 
+struct FOptionalCbField
+{
+	static TOptional<FIoHash> AsHash(FCbFieldView FieldView)
+	{
+		return FieldView.IsHash() ? FieldView.AsHash() : TOptional<FIoHash>();
+	}
+
+	static TOptional<uint64> AsUInt64(FCbFieldView FieldView)
+	{
+		return FieldView.IsInteger() ? FieldView.AsUInt64() : TOptional<uint64>();
+	}
+};
+
 #if HAS_STORAGE_SERVER_RPC_GETCHUNKS_API
-FIoStatus FStorageServerConnection::ReadChunkBatchRequest(const TArray<FChunkBatchRequestEntry>& Chunks, TFunctionRef<void(FIoChunkId Id, EStorageServerContentType MimeType, FIoBuffer Data, uint64 ModTag)> OnResponse, bool bSkipData)
+FIoStatus FStorageServerConnection::ReadChunkBatchRequest(const TArray<FChunkBatchRequestEntry>& Chunks, TFunctionRef<void(FIoChunkId Id, EStorageServerContentType MimeType, FIoBuffer Data, const TOptional<uint64>& ModTag)> OnResponse, bool bSkipData)
 {
 	TRACE_CPUPROFILER_EVENT_SCOPE(FStorageServerConnection::ReadChunkBatchRequest);
 
@@ -342,13 +358,13 @@ FIoStatus FStorageServerConnection::ReadChunkBatchRequest(const TArray<FChunkBat
 		TSet<FIoChunkId> RequestedChunkIds;
 		for (const FChunkBatchRequestEntry& Chunk: Chunks)
 		{
-			if (RequestedChunkIds.Contains(Chunk.Id))
+			if (RequestedChunkIds.Contains(Chunk.ChunkId))
 			{
 				UE_LOG(LogStorageServerConnection, Fatal, TEXT("Multiple requests to the same chunkid are not supported"));
 				return FIoStatus::Unknown;
 			}
 
-			RequestedChunkIds.Add(Chunk.Id);
+			RequestedChunkIds.Add(Chunk.ChunkId);
 		}
 	}
 
@@ -368,9 +384,12 @@ FIoStatus FStorageServerConnection::ReadChunkBatchRequest(const TArray<FChunkBat
 	for (const FChunkBatchRequestEntry& Chunk : Chunks)
 	{
 		Writer.BeginObject();
-		Writer << "Offset" << Chunk.Offset;
-		Writer << "Size" << Chunk.Size;
-		Writer << "Oid" << Chunk.Id;
+		if (!bSkipData)
+		{
+			Writer << "Offset" << Chunk.Offset;
+			Writer << "Size" << Chunk.Size;
+		}
+		Writer << "Oid" << Chunk.ChunkId;
 		if (Chunk.ModTag.IsSet())
 		{
 			Writer << "ModTag" << *Chunk.ModTag;
@@ -468,43 +487,51 @@ FIoStatus FStorageServerConnection::ReadChunkBatchRequest(const TArray<FChunkBat
 		const FIoHash RawHash = Chunk["RawHash"].AsHash();
 		const FIoHash FragmentHash = Chunk["FragmentHash"].AsHash();
 		const FIoHash Hash = Chunk["Hash"].AsHash();
-		const uint64 ModTag = Chunk["ModTag"].AsUInt64();
-
-		if (!ensureMsgf(RawHash != FIoHash() || FragmentHash != FIoHash() || Hash != FIoHash(), TEXT("Failed to find hash in chunk info returned from server")))
-		{
-			continue;
-		}
-
-		const FIoHash& AttachmentHash = (Hash != FIoHash() ? Hash : (RawHash != FIoHash() ? RawHash : FragmentHash));
-
-		int32* InfoIndex = AttachmentHashToIndex.Find(AttachmentHash);
-		if (!ensureMsgf(InfoIndex, TEXT("Failed to find hash in attachments returned from server")))
-		{
-			continue;
-		}
+		const TOptional<uint64> ModTag = FOptionalCbField::AsUInt64(Chunk["ModTag"]);
 
 		FIoChunkId ChunkId;
 		ChunkId.Set(Id.GetView());
 
-		const AttachmentInfo& Attachment = Info[*InfoIndex];
-		ensureAlways(Attachment.Hash == AttachmentHash);
-
-		EStorageServerContentType MimeType = EStorageServerContentType::CompressedBinary;
-		if (Attachment.Flags & AttachmentInfo::IsCompressed)
+		if (bSkipData)
 		{
-			MimeType = EStorageServerContentType::CompressedBinary;
-		}
-		else if (Attachment.Flags & AttachmentInfo::IsObject)
-		{
-			MimeType = EStorageServerContentType::CbObject;
+			OnResponse(ChunkId, EStorageServerContentType::Unknown, FIoBuffer(), ModTag);
 		}
 		else
 		{
-			MimeType = EStorageServerContentType::Binary;
+			if (!ensureMsgf(RawHash != FIoHash() || FragmentHash != FIoHash() || Hash != FIoHash(), TEXT("Failed to find hash in chunk info returned from server")))
+			{
+				continue;
+			}
+
+			const FIoHash& AttachmentHash = (Hash != FIoHash() ? Hash : (RawHash != FIoHash() ? RawHash : FragmentHash));
+
+			int32* InfoIndex = AttachmentHashToIndex.Find(AttachmentHash);
+			if (!ensureMsgf(InfoIndex, TEXT("Failed to find hash in attachments returned from server")))
+			{
+				continue;
+			}
+
+			const AttachmentInfo& Attachment = Info[*InfoIndex];
+			ensureAlways(Attachment.Hash == AttachmentHash);
+
+			EStorageServerContentType MimeType = EStorageServerContentType::CompressedBinary;
+			if (Attachment.Flags & AttachmentInfo::IsCompressed)
+			{
+				MimeType = EStorageServerContentType::CompressedBinary;
+			}
+			else if (Attachment.Flags & AttachmentInfo::IsObject)
+			{
+				MimeType = EStorageServerContentType::CbObject;
+			}
+			else
+			{
+				MimeType = EStorageServerContentType::Binary;
+			}
+
+			FIoBuffer Data(FIoBuffer::Wrap, Result.ValueOrDie().GetData() + InfoOffsets[*InfoIndex], Attachment.PayloadSize);
+			OnResponse(ChunkId, MimeType, Data, ModTag);
 		}
 
-		FIoBuffer Data(FIoBuffer::Wrap, Result.ValueOrDie().GetData() + InfoOffsets[*InfoIndex], Attachment.PayloadSize);
-		OnResponse(ChunkId, MimeType, Data, ModTag);
 		ResponseCount++;
 	}
 
@@ -801,8 +828,6 @@ void FStorageServerConnection::SetupCacheStrategy()
 		return;
 	}
 
-	AsyncQueryLatestServerChunkInfo = MakeShared<FAsyncQueryLatestServerChunkInfo>(*this);
-
 	FString DevStoragePath;
 #if defined(STORAGE_SERVER_PLATFORM_CACHE_ROOT_DIR)
 	DevStoragePath = STORAGE_SERVER_PLATFORM_CACHE_ROOT_DIR;
@@ -819,6 +844,8 @@ void FStorageServerConnection::SetupCacheStrategy()
 	TUniquePtr<StorageServer::ICacheJournal> Journal = MakeUnique<StorageServer::FCacheJournalSimple>(*(DevStoragePath / TEXT("ZenCache.db")), FlushEveryNEntries);
 	TUniquePtr<StorageServer::ICacheStorage> Storage = MakeUnique<StorageServer::FCacheStorageBlocking>(*(DevStoragePath / TEXT("ZenCache.storage")), Size);
 	CacheStrategy = MakeUnique<StorageServer::FCacheStrategyLinear>(MoveTemp(Journal), MoveTemp(Storage), AbandonSize, FlushInterval, bShouldInvalidate);
+
+	AsyncQueryLatestServerChunkInfo = MakeShared<FAsyncQueryLatestServerChunkInfo>(*this);
 }
 #endif // !HAS_STORAGE_SERVER_PLATFORM_CACHE_STRATEGY
 
@@ -845,7 +872,7 @@ bool FStorageServerConnection::FinalizeSetupCacheStrategy()
 
 	// the server chunk info will have been updated by the async init thread.
 	// if there is a platform-specific implementation of SetupCacheStrategy that defers cache creation, the expectation is that it will be handled there
-	AsyncQueryLatestServerChunkInfo.Reset();	
+	AsyncQueryLatestServerChunkInfo.Reset();
 	return true;
 }
 
@@ -865,31 +892,56 @@ FStorageServerConnection::FAsyncQueryLatestServerChunkInfo::~FAsyncQueryLatestSe
 	FPlatformProcess::ReturnSynchEventToPool(IsCompleted);
 }
 
-TArray<StorageServer::FLatestServerChunkInfo> FStorageServerConnection::FAsyncQueryLatestServerChunkInfo::GetResults()
-{
-	IsCompleted->Wait();
-	return MoveTemp(Results);
-}
 bool FStorageServerConnection::FAsyncQueryLatestServerChunkInfo::IsFinished() const
 {
 	return IsCompleted->Wait(0);
 }
 
+void FStorageServerConnection::FAsyncQueryLatestServerChunkInfo::Wait()
+{
+	IsCompleted->Wait();
+}
+
 uint32 FStorageServerConnection::FAsyncQueryLatestServerChunkInfo::Run()
 {
 	TRACE_CPUPROFILER_EVENT_SCOPE(StorageServerRefreshCacheEntries);
 
-	// TODO replace with a new zen server endpoint when it will become available
-	Owner.ChunkInfosRequest([&](FIoChunkId ChunkId, FIoHash RawHash, int64 RawSize)
+#if HAS_STORAGE_SERVER_RPC_GETCHUNKS_API
+	TArray<FChunkBatchRequestEntry> ChunkBatchRequests;
+	Owner.CacheStrategy->IterateChunkIds([&](const FIoChunkId& ChunkId, const StorageServer::FCacheChunkInfo& ChunkInfo)
 	{
-		Results.Emplace(StorageServer::FLatestServerChunkInfo{ChunkId, RawHash, RawSize}); // using RawHash as a subtitute for ModHash.
+		// Force a default ModTag value so we invalidate the chunk if we don't have modtag for some reason
+		const uint64 ModTag = ChunkInfo.ModTag.Get(0);
+
+		ChunkBatchRequests.Emplace(FChunkBatchRequestEntry::VerifyModTagRequest(ChunkId, ModTag));
 	});
 
-	if (Owner.CacheStrategy.IsValid())
+	TSet<FIoChunkId> ValidChunkIds;
+	ValidChunkIds.Reserve(ChunkBatchRequests.Num());
+	Owner.ReadChunkBatchRequest(ChunkBatchRequests, [&](FIoChunkId Id, EStorageServerContentType MimeType, FIoBuffer Data, const TOptional<uint64>& ModTag)
 	{
-		Owner.CacheStrategy->ApplyLatestServerChunkInfo(Results);
-		Results.Reset();
+		if (ModTag.IsSet())
+		{
+			// Chunks rpc endpoint indicate invalid chunks by presence of ModTag in result. 
+			TRACE_CPUPROFILER_EVENT_SCOPE(StorageServerRefreshCacheEntries::Invalidate);
+			Owner.CacheStrategy->Invalidate(Id);
+		}
+		else
+		{
+			ValidChunkIds.Add(Id);
+		}
+	}, true);
+
+	for (const FChunkBatchRequestEntry& Request: ChunkBatchRequests)
+	{
+		if (!ValidChunkIds.Contains(Request.ChunkId))
+		{
+			// Absence of chunkid in response means chunk is not present on the server 
+			TRACE_CPUPROFILER_EVENT_SCOPE(StorageServerRefreshCacheEntries::Invalidate);
+			Owner.CacheStrategy->Invalidate(Request.ChunkId);
+		}
 	}
+#endif
 
 	IsCompleted->Trigger();
 	return 0;
diff --git a/Engine/Source/Runtime/StorageServerClient/Private/StorageServerConnection.h b/Engine/Source/Runtime/StorageServerClient/Private/StorageServerConnection.h
--- a/Engine/Source/Runtime/StorageServerClient/Private/StorageServerConnection.h
+++ b/Engine/Source/Runtime/StorageServerClient/Private/StorageServerConnection.h
@@ -67,15 +67,26 @@ public:
 #if HAS_STORAGE_SERVER_RPC_GETCHUNKS_API
 	// Matches input parameters for ProjectStore::GetChunks in zen server
 	struct FChunkBatchRequestEntry {
-		FIoChunkId Id;
+		FIoChunkId ChunkId;
 		uint64 Offset;
 		uint64 Size;
 		TOptional<uint64> ModTag;
+
+		static FChunkBatchRequestEntry DataRequest(const FIoChunkId& ChunkId, const uint64 Offset, const uint64 Size)
+		{
+			TOptional<uint64> EmptyModTag;
+			return FChunkBatchRequestEntry {ChunkId, Offset, Size, EmptyModTag};
+		}
+
+		static FChunkBatchRequestEntry VerifyModTagRequest(const FIoChunkId& ChunkId, const uint64 ModTag)
+		{
+			return FChunkBatchRequestEntry {ChunkId, 0, (uint64)-1, ModTag};
+		}
 	};
 
 	FIoStatus ReadChunkBatchRequest(
 		const TArray<FChunkBatchRequestEntry>& Chunks,
-		TFunctionRef<void(FIoChunkId Id, EStorageServerContentType MimeType, FIoBuffer Data, uint64 ModTag)> OnResponse,
+		TFunctionRef<void(FIoChunkId Id, EStorageServerContentType MimeType, FIoBuffer Data, const TOptional<uint64>& ModTag)> OnResponse,
 		bool bSkipData = false // if bSkipData is true, then OnResponse will only be called for requested chunks that have either empty or different modtags
 	);
 #endif
@@ -124,16 +135,14 @@ private:
 		virtual ~FAsyncQueryLatestServerChunkInfo();
 		
 		bool IsFinished() const;
-		TArray<StorageServer::FLatestServerChunkInfo> GetResults();
-		
+		void Wait();
 	private:
 		virtual uint32 Run() override;
 
 		FStorageServerConnection& Owner;
 		class FEvent* IsCompleted;
-		TArray<StorageServer::FLatestServerChunkInfo> Results;
-
 	};
+
 	TSharedPtr<FAsyncQueryLatestServerChunkInfo, ESPMode::ThreadSafe> AsyncQueryLatestServerChunkInfo;
 };
 
diff --git a/Engine/Source/Runtime/StorageServerClient/Private/StorageServerPlatformFile.cpp b/Engine/Source/Runtime/StorageServerClient/Private/StorageServerPlatformFile.cpp
--- a/Engine/Source/Runtime/StorageServerClient/Private/StorageServerPlatformFile.cpp
+++ b/Engine/Source/Runtime/StorageServerClient/Private/StorageServerPlatformFile.cpp
@@ -339,7 +339,7 @@ private:
 		return 0;
 	}
 
-	void HandleChunkBatchResponse(FIoChunkId ChunkId, EStorageServerContentType MimeType, FIoBuffer Data, uint64 ModTag)
+	void HandleChunkBatchResponse(FIoChunkId ChunkId, EStorageServerContentType MimeType, FIoBuffer Data, const TOptional<uint64>& ModTag)
 	{
 		TArray<uint8> DecompressedData;
 
@@ -396,7 +396,7 @@ private:
 			return;
 		}
 
-		FIoStatus ResponseStatus = Connection.ReadChunkBatchRequest(Items, [&](FIoChunkId Id, EStorageServerContentType MimeType, FIoBuffer Data, uint64 ModTag)
+		FIoStatus ResponseStatus = Connection.ReadChunkBatchRequest(Items, [&](FIoChunkId Id, EStorageServerContentType MimeType, FIoBuffer Data, const TOptional<uint64>& ModTag)
 		{
 			HandleChunkBatchResponse(Id, MimeType, Data, ModTag);
 		});
