Way to enable zen target side caching via command line arguments

diff --git a/Engine/Source/Runtime/StorageServerClient/Private/StorageServerConnection.cpp b/Engine/Source/Runtime/StorageServerClient/Private/StorageServerConnection.cpp
--- a/Engine/Source/Runtime/StorageServerClient/Private/StorageServerConnection.cpp
+++ b/Engine/Source/Runtime/StorageServerClient/Private/StorageServerConnection.cpp
@@ -31,16 +31,6 @@
 
 #if !UE_BUILD_SHIPPING
 
-// Set to true to enable cache, temporarily disabled for now
-static constexpr bool bEnableZenCache = false; // PLATFORM_SUPPORTS_STORAGE_SERVER_CACHE;
-
-// TODO implement a way to control cache via cvars, right now it's not possible as cvars are initialized at later stage than cache
-extern "C" int32 GStorageServerLocalCacheSizeKB = bEnableZenCache ? 8 * 1024 * 1024 : 0; // 8 GB
-extern "C" float GStorageServerLocalCacheFlushInterval = 10.0;
-extern "C" int32 GStorageServerLocalCacheFlushEveryNEntries = 0; // set to >0 to flush journal every N new entries
-extern "C" int32 GStorageServerLocalCacheAbandonSizeKB = 512 * 1024; // abandon cache if >= 512 MB of it is invalid
-extern "C" int32 GStorageServerLocalCacheInvalidate = 0; // set to 1 to invalidate
-
 DEFINE_LOG_CATEGORY(LogStorageServerConnection);
 
 TRACE_DECLARE_INT_COUNTER(ZenHttpClientSerializedBytes, TEXT("ZenClient/SerializedBytes (compressed)"));
@@ -819,10 +809,42 @@ bool FStorageServerConnection::HandshakeRequest()
 	return false;
 }
 
+void FStorageServerConnection::GetDefaultCacheConfiguration(FCacheConfiguration& OutConfiguration)
+{
+	const TCHAR* CmdLine = FCommandLine::Get();
+
+	OutConfiguration.bEnable = PLATFORM_SUPPORTS_STORAGE_SERVER_CACHE && (FCString::Strstr(CmdLine, TEXT("-ZenEnableCache")) != nullptr);
+	if (!OutConfiguration.bEnable)
+	{
+		OutConfiguration.bForceInvalidate = false;
+		OutConfiguration.CacheSizeKB = 0;
+		OutConfiguration.AbandonSizeKB = 0;
+		OutConfiguration.FlushInterval = 0.0f;
+		OutConfiguration.FlushEveryNEntries = 0;
+		return;
+	}
+
+	OutConfiguration.bForceInvalidate = FCString::Strstr(CmdLine, TEXT("-ZenInvalidateCache")) != nullptr;
+	if (!FParse::Value(CmdLine, TEXT("-ZenCacheSizeKB="), OutConfiguration.CacheSizeKB))
+	{
+		OutConfiguration.CacheSizeKB = 4 * 1024 * 1024; // 4 GB by default
+	}
+	if (!FParse::Value(CmdLine, TEXT("-ZenCacheAbandonSizeKB="), OutConfiguration.AbandonSizeKB))
+	{
+		OutConfiguration.AbandonSizeKB = 512 * 1024; // 512 MB by default
+	}
+	OutConfiguration.FlushInterval = 10.0f;
+	OutConfiguration.FlushEveryNEntries = 0;
+}
+
 #if !HAS_STORAGE_SERVER_PLATFORM_CACHE_STRATEGY
+
 void FStorageServerConnection::SetupCacheStrategy()
 {
-	if (GStorageServerLocalCacheSizeKB <= 0)
+	FCacheConfiguration Configuration = {};
+	GetDefaultCacheConfiguration(Configuration);
+	
+	if (!Configuration.bEnable)
 	{
 		CacheStrategy.Reset();
 		return;
@@ -835,11 +857,11 @@ void FStorageServerConnection::SetupCacheStrategy()
 	DevStoragePath = FPlatformMisc::GamePersistentDownloadDir();
 #endif
 
-	const uint64 Size = ((uint64)GStorageServerLocalCacheSizeKB) * 1024;
-	const float FlushInterval = GStorageServerLocalCacheFlushInterval;
-	const uint64 FlushEveryNEntries = GStorageServerLocalCacheFlushEveryNEntries;
-	const uint64 AbandonSize = ((uint64)GStorageServerLocalCacheAbandonSizeKB) * 1024;
-	const bool bShouldInvalidate = GStorageServerLocalCacheInvalidate != 0;
+	const uint64 Size = ((uint64)Configuration.CacheSizeKB) * 1024;
+	const float FlushInterval = Configuration.FlushInterval;
+	const uint64 FlushEveryNEntries = Configuration.FlushEveryNEntries;
+	const uint64 AbandonSize = ((uint64)Configuration.AbandonSizeKB) * 1024;
+	const bool bShouldInvalidate = Configuration.bForceInvalidate;
 
 	TUniquePtr<StorageServer::ICacheJournal> Journal = MakeUnique<StorageServer::FCacheJournalSimple>(*(DevStoragePath / TEXT("ZenCache.db")), FlushEveryNEntries);
 	TUniquePtr<StorageServer::ICacheStorage> Storage = MakeUnique<StorageServer::FCacheStorageBlocking>(*(DevStoragePath / TEXT("ZenCache.storage")), Size);
diff --git a/Engine/Source/Runtime/StorageServerClient/Private/StorageServerConnection.h b/Engine/Source/Runtime/StorageServerClient/Private/StorageServerConnection.h
--- a/Engine/Source/Runtime/StorageServerClient/Private/StorageServerConnection.h
+++ b/Engine/Source/Runtime/StorageServerClient/Private/StorageServerConnection.h
@@ -115,6 +115,17 @@ private:
 	TUniquePtr<IStorageServerHttpClient> CreateHttpClient(const FString Address, const int32 Port);
 	TSharedPtr<FInternetAddr> StringToInternetAddr(const FString Address, const int32 Port);
 	bool HandshakeRequest();
+
+	struct FCacheConfiguration
+	{
+		bool bEnable;				// set to true to enable cache
+		bool bForceInvalidate;		// invalidate cache if set to true
+		int32 CacheSizeKB;			// total size of cache in kb
+		float FlushInterval;		// interval at which to flush cache in seconds
+		int32 FlushEveryNEntries;	// set to >0 to flush journal every N new entries
+		int32 AbandonSizeKB;		// set to >0 to abandon cache if amount of invalid data goes over threshold
+	};
+	void GetDefaultCacheConfiguration(FCacheConfiguration& OutConfiguration);
 	void SetupCacheStrategy();
 	bool FinalizeSetupCacheStrategy();
 	void BuildReadChunkRequestUrl(FAnsiStringBuilderBase& Builder, const FIoChunkId& ChunkId, const uint64 Offset, const uint64 Size);
