Add a command line parameter to enable zen cache sectioned journal

#rnx

diff --git a/Engine/Source/Runtime/StorageServerClient/Private/StorageServerConnection.cpp b/Engine/Source/Runtime/StorageServerClient/Private/StorageServerConnection.cpp
--- a/Engine/Source/Runtime/StorageServerClient/Private/StorageServerConnection.cpp
+++ b/Engine/Source/Runtime/StorageServerClient/Private/StorageServerConnection.cpp
@@ -26,7 +26,9 @@
 #include "HAL/Event.h"
 
 #include "Cache/CacheJournalSimple.h"
+#include "Cache/CacheJournalSectioned.h"
 #include "Cache/CacheStorageBlocking.h"
+#include "Cache/CacheStorageMmap.h"
 #include "Cache/CacheStrategyLinear.h"
 
 #if !UE_BUILD_SHIPPING
@@ -816,11 +818,6 @@ void FStorageServerConnection::GetDefaultCacheConfiguration(FCacheConfiguration&
 	OutConfiguration.bEnable = PLATFORM_SUPPORTS_STORAGE_SERVER_CACHE && (FCString::Strstr(CmdLine, TEXT("-ZenEnableCache")) != nullptr);
 	if (!OutConfiguration.bEnable)
 	{
-		OutConfiguration.bForceInvalidate = false;
-		OutConfiguration.CacheSizeKB = 0;
-		OutConfiguration.AbandonSizeKB = 0;
-		OutConfiguration.FlushInterval = 0.0f;
-		OutConfiguration.FlushEveryNEntries = 0;
 		return;
 	}
 
@@ -835,6 +832,9 @@ void FStorageServerConnection::GetDefaultCacheConfiguration(FCacheConfiguration&
 	}
 	OutConfiguration.FlushInterval = 10.0f;
 	OutConfiguration.FlushEveryNEntries = 0;
+
+	OutConfiguration.bUseSectionedJournal = FCString::Strstr(CmdLine, TEXT("-ZenUseSectionedJournal")) != nullptr;
+	OutConfiguration.bUseMemoryMappedStorage = FPlatformProperties::SupportsMemoryMappedFiles() && (FCString::Strstr(CmdLine, TEXT("-ZenMmapedStorage")) != nullptr);
 }
 
 #if !HAS_STORAGE_SERVER_PLATFORM_CACHE_STRATEGY
@@ -863,8 +863,27 @@ void FStorageServerConnection::SetupCacheStrategy()
 	const uint64 AbandonSize = ((uint64)Configuration.AbandonSizeKB) * 1024;
 	const bool bShouldInvalidate = Configuration.bForceInvalidate;
 
-	TUniquePtr<StorageServer::ICacheJournal> Journal = MakeUnique<StorageServer::FCacheJournalSimple>(*(DevStoragePath / TEXT("ZenCache.db")), FlushEveryNEntries);
-	TUniquePtr<StorageServer::ICacheStorage> Storage = MakeUnique<StorageServer::FCacheStorageBlocking>(*(DevStoragePath / TEXT("ZenCache.storage")), Size);
+	TUniquePtr<StorageServer::ICacheJournal> Journal;
+	if (Configuration.bUseSectionedJournal)
+	{
+		Journal = MakeUnique<StorageServer::FCacheJournalSectioned>(*(DevStoragePath / TEXT("ZenCacheJournalSectioned.db")));
+	}
+	else
+	{
+		Journal = MakeUnique<StorageServer::FCacheJournalSimple>(*(DevStoragePath / TEXT("ZenCache.db")), FlushEveryNEntries);
+	}
+
+	TUniquePtr<StorageServer::ICacheStorage> Storage;
+
+	if (Configuration.bUseMemoryMappedStorage)
+	{
+		Storage = MakeUnique<StorageServer::FCacheStorageMmap>(*(DevStoragePath / TEXT("ZenCache.storage")), Size);
+	}
+	else
+	{
+		Storage = MakeUnique<StorageServer::FCacheStorageBlocking>(*(DevStoragePath / TEXT("ZenCache.storage")), Size);
+	}
+
 	CacheStrategy = MakeUnique<StorageServer::FCacheStrategyLinear>(MoveTemp(Journal), MoveTemp(Storage), AbandonSize, FlushInterval, bShouldInvalidate);
 
 	AsyncQueryLatestServerChunkInfo = MakeShared<FAsyncQueryLatestServerChunkInfo>(*this);
diff --git a/Engine/Source/Runtime/StorageServerClient/Private/StorageServerConnection.h b/Engine/Source/Runtime/StorageServerClient/Private/StorageServerConnection.h
--- a/Engine/Source/Runtime/StorageServerClient/Private/StorageServerConnection.h
+++ b/Engine/Source/Runtime/StorageServerClient/Private/StorageServerConnection.h
@@ -118,12 +118,14 @@ private:
 
 	struct FCacheConfiguration
 	{
-		bool bEnable;				// set to true to enable cache
-		bool bForceInvalidate;		// invalidate cache if set to true
-		int32 CacheSizeKB;			// total size of cache in kb
-		float FlushInterval;		// interval at which to flush cache in seconds
-		int32 FlushEveryNEntries;	// set to >0 to flush journal every N new entries
-		int32 AbandonSizeKB;		// set to >0 to abandon cache if amount of invalid data goes over threshold
+		bool bEnable = false;					// set to true to enable cache
+		bool bForceInvalidate = false;			// invalidate cache if set to true
+		int32 CacheSizeKB = 0;					// total size of cache in kb
+		float FlushInterval = 0.f;				// interval at which to flush cache in seconds
+		int32 FlushEveryNEntries = 0;			// set to >0 to flush journal every N new entries
+		int32 AbandonSizeKB = 0;				// set to >0 to abandon cache if amount of invalid data goes over threshold
+		bool bUseSectionedJournal = false;		// use the sectioned journal instead of the simple TMap variant
+		bool bUseMemoryMappedStorage = false;	// use mmapped cache storage backend
 	};
 	void GetDefaultCacheConfiguration(FCacheConfiguration& OutConfiguration);
 	void SetupCacheStrategy();
