Fixing broken read in zen mmap cache storage

#rnx

diff --git a/Engine/Source/Runtime/StorageServerClient/Private/Cache/CacheStorageMmap.cpp b/Engine/Source/Runtime/StorageServerClient/Private/Cache/CacheStorageMmap.cpp
--- a/Engine/Source/Runtime/StorageServerClient/Private/Cache/CacheStorageMmap.cpp
+++ b/Engine/Source/Runtime/StorageServerClient/Private/Cache/CacheStorageMmap.cpp
@@ -104,7 +104,7 @@ namespace StorageServer
 
 		if (SizeA > 0 && SizeB > 0)
 		{
-			FIoBuffer Result = (OptDestination.IsSet() && OptDestination->GetSize() >= SizeA + SizeB) ? OptDestination.GetValue() : FIoBuffer();
+			FIoBuffer Result = (OptDestination.IsSet() && OptDestination->GetSize() >= SizeA + SizeB) ? OptDestination.GetValue() : FIoBuffer(SizeA + SizeB);
 			Result.SetSize(SizeA + SizeB);
 			FMemory::Memcpy(Result.GetData(), BackingFiles[IndexA].MapPtr + OffsetA, SizeA);
 			FMemory::Memcpy(Result.GetData() + SizeA, BackingFiles[IndexB].MapPtr + OffsetB, SizeB);
