Refactor zen cache sectioned journal to use FArchive serialization, added unit tests

diff --git a/Engine/Source/Runtime/StorageServerClient/Private/Cache/CacheJournal.h b/Engine/Source/Runtime/StorageServerClient/Private/Cache/CacheJournal.h
--- a/Engine/Source/Runtime/StorageServerClient/Private/Cache/CacheJournal.h
+++ b/Engine/Source/Runtime/StorageServerClient/Private/Cache/CacheJournal.h
@@ -93,7 +93,7 @@ namespace StorageServer
 		virtual ~ICacheJournal() = default;
 
 		// Flushes data to backing storage
-		virtual void Flush() = 0;
+		virtual void Flush(bool bImmediate) = 0;
 
 		// Invalidates all data in journal
 		virtual void InvalidateAll() = 0;
diff --git a/Engine/Source/Runtime/StorageServerClient/Private/Cache/CacheJournalSectioned.cpp b/Engine/Source/Runtime/StorageServerClient/Private/Cache/CacheJournalSectioned.cpp
--- a/Engine/Source/Runtime/StorageServerClient/Private/Cache/CacheJournalSectioned.cpp
+++ b/Engine/Source/Runtime/StorageServerClient/Private/Cache/CacheJournalSectioned.cpp
@@ -2,75 +2,101 @@
 
 #include "CacheJournalSectioned.h"
 
+#if !UE_BUILD_SHIPPING
+
 #include "HAL/PlatformFile.h"
 #include "Misc/ScopeLock.h"
 #include "Serialization/MemoryWriter.h"
-
-#if !UE_BUILD_SHIPPING
-
-template<typename M>
-class TMapEnum
-{
-public:
-	M& Data;
-
-	TMapEnum(M& Item)
-		: Data(Item)
-	{
-	}
-
-	TMapEnum& operator()(typename M::KeyType const& Key, typename M::ValueType Value)
-	{
-		Data.Add(Key, Value);
-		return *this;
-	}
-};
-
-template<typename T>
-TMapEnum<T> TMapInit(T& Item)
-{
-	return TMapEnum<T>(Item);
-}
-
+#include "ProfilingDebugging/CountersTrace.h"
 
 DEFINE_LOG_CATEGORY_STATIC(LogCacheJournal, Log, All);
 
+TRACE_DECLARE_INT_COUNTER(ZenJournalSectionedChunks, TEXT("ZenClient/JournalSectioned/Chunks"));
+TRACE_DECLARE_INT_COUNTER(ZenJournalSectionedEntries, TEXT("ZenClient/JournalSectioned/Entries"));
+
 namespace StorageServer
 {
-	static const uint64 PAGE_MAGIC = 0x5041'4745; // PAGE
-	static const uint32 JOURNAL_VERSION = 0x01;
+	static inline constexpr uint64 FILE_MAGIC = 0x5a45'4e43; // ZENC
+	static inline constexpr uint64 PAGE_MAGIC = 0x5041'4745; // PAGE
+	static inline constexpr uint32 JOURNAL_VERSION = 0x02;
+
+	// TODO hardcoded constants of max entries are fragile, can we compute it from page size directly?
+	static inline constexpr uint32 GMaxChunkInfoEntries = 2950; // approx 44 bytes per entry
+	static inline constexpr uint32 GMaxChunkEntries = 1880; // approx 66 bytes per entry
+	static inline constexpr uint32 JOURNAL_PAGE_SIZE = 128 * 1024;
 
 	FCacheJournalSectioned::FCacheJournalSectioned(FStringView InJournalPath)
 		: JournalFilePath(InJournalPath)
 	{
 		JournalFileHandle.Reset(IPlatformFile::GetPlatformPhysical().OpenWrite(*JournalFilePath, true, true));
 
-		LoadCacheJournal(JournalFilePath);
+		LoadCacheJournal();
 	}
 
-	void FCacheJournalSectioned::Flush()
+	void FCacheJournalSectioned::Flush(bool bImmediate)
 	{
 		TRACE_CPUPROFILER_EVENT_SCOPE(FCacheJournalSectioned::Flush);
 
+		if (bImmediate)
+		{
+			FlushImmediate();
+		}
+		else
+		{
+			UE::TScopeLock _(CriticalSection);
+
+			if (FlushTask.IsCompleted())
+			{
+				FlushTask = UE::Tasks::Launch(TEXT("CacheJournalFlush"), [this]() {FlushImmediate();});
+			}
+		}
+	}
+
+	void FCacheJournalSectioned::FlushImmediate()
+	{
+		TRACE_CPUPROFILER_EVENT_SCOPE(FCacheJournalSectioned::FlushImmediate);
+
 		UE::TScopeLock _(CriticalSection);
 
+		TArray<uint8> SerializationBuffer;
+
 		if (bPagesModified)
 		{
-			FJournalHeader Header{ .Version = JOURNAL_VERSION, .NumberOfChunkPages = ChunkPages.Num(), .NumberOfChunkInfoPages = ChunkInfoPages.Num() };
-			JournalFileHandle->Seek(0);
-			JournalFileHandle->Write((uint8*)&Header, sizeof(Header));
+			FJournalHeader Header;
+			Header.Magic = FILE_MAGIC;
+			Header.Version = JOURNAL_VERSION;
+			Header.PageCount = ChunkPages.Num() + ChunkInfoPages.Num();
+
+			SerializationBuffer.SetNum(0, EAllowShrinking::No);
+			FMemoryWriter Ar(SerializationBuffer);
+			Ar << Header;
+
+			if (!ensure(SerializationBuffer.Num() <= FJournalHeader::SerializedSize))
+			{
+				return;
+			}
+
+			if (!JournalFileHandle->Seek(0))
+			{
+				return;
+			}
+
+			if (!JournalFileHandle->Write(SerializationBuffer.GetData(), SerializationBuffer.Num()))
+			{
+				return;
+			}
 
 			bPagesModified = false;
 		}
 
 		for (FJournalChunkPage& Page : ChunkPages)
 		{
-			Page.Flush(JournalFileHandle.Get());
+			Page.Flush(JournalFileHandle.Get(), SerializationBuffer);
 		}
 
 		for (FJournalChunkInfoPage& Page : ChunkInfoPages)
 		{
-			Page.Flush(JournalFileHandle.Get());
+			Page.Flush(JournalFileHandle.Get(), SerializationBuffer);
 		}
 	}
 
@@ -78,15 +104,12 @@ namespace StorageServer
 	{
 		UE::TScopeLock _(CriticalSection);
 
-		for (FJournalChunkPage& Page : ChunkPages)
-		{
-			Page.InvalidateAll();
-		}
+		ChunkPages.Empty();
+		ChunkInfoPages.Empty();
+		ResetJournalFilePos();
 
-		for (FJournalChunkInfoPage& InfoPage : ChunkInfoPages)
-		{
-			InfoPage.InvalidateAll();
-		}
+		TRACE_COUNTER_SET(ZenJournalSectionedChunks, 0);
+		TRACE_COUNTER_SET(ZenJournalSectionedEntries, 0);
 
 		bPagesModified = true;
 	}
@@ -121,11 +144,12 @@ namespace StorageServer
 			}
 		}
 
-		FJournalChunkInfoPage* Page = FindAddChunkInfoPage(ChunkId);
+		FJournalChunkInfoPage* Page = FindOrAddChunkInfoPage(ChunkId);
 		check(Page != nullptr && !Page->IsFull());
 
 		ChunkInfo.SetChunkInfo(OptModHash, OptRawSize, OptRawBlockSize);
 		Page->Add(ChunkId, ChunkInfo);
+		TRACE_COUNTER_INCREMENT(ZenJournalSectionedEntries);
 		return true;
 	}
 
@@ -146,8 +170,6 @@ namespace StorageServer
 
 	bool FCacheJournalSectioned::AddEntry(const FIoChunkId& ChunkId, const FCacheEntry& Entry)
 	{
-		EJournalPageResult Res;
-
 		UE::TScopeLock _(CriticalSection);
 
 		FCacheEntry OutEntry;
@@ -156,10 +178,10 @@ namespace StorageServer
 			return false;
 		}
 
-		FJournalChunkPage* Page = FindChunkPage(ChunkId);
-
-		Res = Page->Add(ChunkId, Entry);
+		FJournalChunkPage* Page = FindOrAddChunkPage(ChunkId);
+		check(Page != nullptr && !Page->IsFull());
 
+		EJournalPageResult Res = Page->Add(ChunkId, Entry);
 		return Res == EJournalPageResult::Ok;
 	}
 
@@ -208,141 +230,138 @@ namespace StorageServer
 		}
 	}
 
-	void FCacheJournalSectioned::LoadCacheJournal(FStringView JournalPath)
+	bool FCacheJournalSectioned::LoadCacheJournal()
 	{
 		TRACE_CPUPROFILER_EVENT_SCOPE(FCacheJournalSectioned::LoadCacheJournal);
 
 		UE::TScopeLock _(CriticalSection);
 
-		check(JournalFileHandle.IsValid());
-
-		const int64 FileSize = JournalFileHandle->Size();
-		if (FileSize < sizeof(FJournalHeader))
+		if (!JournalFileHandle.IsValid() || JournalFileHandle->Size() < FJournalHeader::SerializedSize)
 		{
-			return;
-		}
-
-		FJournalHeader Header;
-		JournalFileHandle->ReadAt((uint8*)&Header, sizeof(Header), 0);
-		if (Header.Version != JOURNAL_VERSION)
-		{
-			UE_LOG(LogCacheJournal, Error, TEXT("Journal file is of different version %d, supported %d"), Header.Version, JOURNAL_VERSION);
-			return;
+			return false;
 		}
 
-		int64 ReadCursor = sizeof(Header);
-		int32 ChunkEntryCount = 0;
-		int32 ChunkInfoEntryCount = 0;
-		int32 TotalMBs = 0;
+		TArray<uint8> SerializationBuffer;
+		SerializationBuffer.Reserve(JOURNAL_PAGE_SIZE);
 
-		TMap<EPageType, FString> TypeArray;
-		TMapInit(TypeArray)
-			(EPageType::Chunk, TEXT("Chunk"))
-			(EPageType::ChunkInfo, TEXT("ChunkInfo"));
+		FJournalHeader Header;
 
-		while ((uint32)(ReadCursor + sizeof(FJournalPageDescriptor)) < FileSize)
 		{
-			auto ReadAndProgress = [&JournalFileHandle = JournalFileHandle, &ReadCursor](uint8* Destination, uint32 Size) -> bool
+			SerializationBuffer.SetNum(FJournalHeader::SerializedSize, EAllowShrinking::No);
+
+			if (!JournalFileHandle->ReadAt(SerializationBuffer.GetData(), SerializationBuffer.Num(), 0))
 			{
-				if (!JournalFileHandle->ReadAt(Destination, Size, ReadCursor))
-				{
-					return false;
-				}
-				ReadCursor += Size;
-				return true;
-			};
+				return false;
+			}
 
-			int32 PageFilePos = ReadCursor;
-			FJournalPageDescriptor SectionDescriptor;
+			FMemoryReader HeaderAr(SerializationBuffer);
+			HeaderAr << Header;
 
-			if (!ReadAndProgress((uint8*)&SectionDescriptor, sizeof(SectionDescriptor)))
+			if (Header.Magic != FILE_MAGIC)
 			{
-				UE_LOG(LogCacheJournal, Error, TEXT("Failed to read section from journal file"));
+				UE_LOG(LogCacheJournal, Warning, TEXT("Journal file has wrong magic number, the journal is corrupt"));
 				InvalidateAll();
-				return;
+				return false;
 			}
 
-			if (SectionDescriptor.Magic != PAGE_MAGIC)
+			if (Header.Version != JOURNAL_VERSION)
 			{
-				UE_LOG(LogCacheJournal, Error, TEXT("Section has wrong magic number, the journal is corrupt"));
+				UE_LOG(LogCacheJournal, Warning, TEXT("Journal file is of different version %d, supported %d"), Header.Version, JOURNAL_VERSION);
 				InvalidateAll();
-				return;
+				return false;
 			}
+		}
 
-			if (SectionDescriptor.Type == EPageType::Chunk)
-			{
-				int32 EntriesCount = SectionDescriptor.Size / sizeof(FJournalStoreEntry);
-				TArray<FJournalStoreEntry> Entries;
-				Entries.AddUninitialized(EntriesCount);
+		int64 ReadCursor = FJournalHeader::SerializedSize;
+		uint64 PageCount = 0;
 
-				ChunkEntryCount += EntriesCount;
+		uint64 ChunkInfoEntryCount = 0;
+		uint64 ChunkEntryCount = 0;
 
-				if (!ReadAndProgress((uint8*)Entries.GetData(), SectionDescriptor.Size))
-				{
-					UE_LOG(LogCacheJournal, Error, TEXT("Failed to read chunks from section in journal file"));
-					InvalidateAll();
-					return;
-				}
+		while ((ReadCursor + FJournalPageHeader::SerializedSize < JournalFileHandle->Size()) && (PageCount < Header.PageCount))
+		{
+			const int64 PageOffset = ReadCursor;
 
-				int64 AllChunksSize = 0;
-				int64 NumberOfNoSizeChunks = 0;
-				for (FJournalStoreEntry& Entry : Entries)
-				{
-					if (Entry.Entry.ChunkSize <= 0)
-					{
-						NumberOfNoSizeChunks++;
-						continue;
-					}
+			FJournalPageHeader PageHeader;
+			SerializationBuffer.SetNum(FJournalPageHeader::SerializedSize, EAllowShrinking::No);
 
-					AllChunksSize += Entry.Entry.ChunkSize;
-				}
+			if (!JournalFileHandle->ReadAt(SerializationBuffer.GetData(), SerializationBuffer.Num(), PageOffset))
+			{
+				UE_LOG(LogCacheJournal, Warning, TEXT("Failed to read page header from section in journal file"));
+				return false;
+			}
 
-				const int32 SizeInMb = AllChunksSize >> 20;
-				TotalMBs += SizeInMb;
-				UE_LOG(LogCacheJournal, Log, TEXT("Found %s page with %d entries of total size %d, Number of chunks with size -1 %d"), 
-					*TypeArray[SectionDescriptor.Type], EntriesCount, SizeInMb, NumberOfNoSizeChunks);
+			FMemoryReader HeaderAr(SerializationBuffer);
+			HeaderAr << PageHeader;
 
-				FJournalChunkPage& Page = ChunkPages.Add_GetRef({ GMaxSectionEntries });
-				Page.SetFilePos(PageFilePos);
-				Page << MoveTemp(Entries);
+			if (PageHeader.Magic != PAGE_MAGIC)
+			{
+				UE_LOG(LogCacheJournal, Warning, TEXT("Section has wrong magic number, the journal is corrupt"));
+				InvalidateAll();
+				return false;
 			}
-			else if (SectionDescriptor.Type == EPageType::ChunkInfo)
+
+			if (PageHeader.PageSize > JOURNAL_PAGE_SIZE || PageHeader.DataSize >= PageHeader.PageSize)
 			{
-				TArray<FJournalChunkInfoPage::FChunkInfoEntry> ChunkInfoEntries;
-				int32 EntriesCount = SectionDescriptor.Size / sizeof(FJournalChunkInfoPage::FChunkInfoEntry);
-				ChunkInfoEntries.AddUninitialized(EntriesCount);
+				UE_LOG(LogCacheJournal, Warning, TEXT("Section has wrong size, the journal is corrupt"));
+				InvalidateAll();
+				return false;
+			}
 
-				ChunkInfoEntryCount += EntriesCount;
+			SerializationBuffer.SetNum(PageHeader.DataSize, EAllowShrinking::No);
 
-				if (!ReadAndProgress((uint8*)ChunkInfoEntries.GetData(), SectionDescriptor.Size))
-				{
-					UE_LOG(LogCacheJournal, Error, TEXT("Failed to read chunk infos from section in journal file"));
-					InvalidateAll();
-					return;
-				}
+			if (!JournalFileHandle->ReadAt(SerializationBuffer.GetData(), SerializationBuffer.Num(), PageOffset + FJournalPageHeader::SerializedSize))
+			{
+				UE_LOG(LogCacheJournal, Warning, TEXT("Failed to read page from section in journal file."));
+				InvalidateAll();
+				return false;
+			}
 
-				FJournalChunkInfoPage& Page = ChunkInfoPages.Add_GetRef({ GMaxSectionEntries });
-				Page.SetFilePos(PageFilePos);
-				Page << MoveTemp(ChunkInfoEntries);
+			FJournalPageBase* Page = nullptr;
+			switch (PageHeader.Type)
+			{
+			case EPageType::Chunk:
+				Page = &ChunkPages.Add_GetRef({ PageHeader.PageSize, GMaxChunkEntries }); // TODO GMaxChunkEntries ?
+				break;
+			case EPageType::ChunkInfo:
+				Page = &ChunkInfoPages.Add_GetRef({ PageHeader.PageSize, GMaxChunkInfoEntries }); // TODO GMaxChunkInfoEntries ?
+				break;
+			default:
+				UE_LOG(LogCacheJournal, Warning, TEXT("Unknown page type in journal file."));
+				InvalidateAll();
+				return false;
 			}
-			
-			int32 BytesLeft = SectionDescriptor.MaxSize - SectionDescriptor.Size;
-			if (BytesLeft > 0)
+
+			Page->SetFilePos(PageOffset);
+
+			FMemoryReader BodyAr(SerializationBuffer);
+			BodyAr << *Page;
+
+			ReadCursor += PageHeader.PageSize;
+			PageCount++;
+
+			switch (PageHeader.Type)
 			{
-				ReadCursor += BytesLeft;
+			case EPageType::Chunk:
+				ChunkEntryCount += ((FJournalChunkPage*)Page)->GetEntryCount();
+				break;
+			case EPageType::ChunkInfo:
+				ChunkInfoEntryCount += ((FJournalChunkInfoPage*)Page)->GetEntryCount();
+				break;
+			default:
+				break;
 			}
 		}
 
-		check(Header.NumberOfChunkInfoPages == ChunkInfoPages.Num());
-		check(Header.NumberOfChunkPages == ChunkPages.Num());
+		TRACE_COUNTER_SET(ZenJournalSectionedChunks, ChunkEntryCount);
+		TRACE_COUNTER_SET(ZenJournalSectionedEntries, ChunkInfoEntryCount);
+
+		UE_LOG(LogCacheJournal, Display, TEXT("Cache journal loaded. Imported %d chunk pages with %d entries and %d chunk info pages with %d info entries."), ChunkPages.Num(), ChunkEntryCount, ChunkInfoPages.Num(), ChunkInfoEntryCount);
 
-		JournalSectionCursor = ReadCursor;
-		NextAvailableJournalPos = ReadCursor;
-		UE_LOG(LogCacheJournal, Display, TEXT("Cache journal loaded. Imported %d chunk pages with %d entries and %d chunk info pages with %d info entries. Total size of data referenced by journal %d MB"),
-			ChunkPages.Num(), ChunkEntryCount, ChunkInfoPages.Num(), ChunkInfoEntryCount, TotalMBs);
+		return true;
 	}
 
-	FJournalChunkInfoPage* FCacheJournalSectioned::FindAddChunkInfoPage(const FIoChunkId& ChunkId)
+	FJournalChunkInfoPage* FCacheJournalSectioned::FindOrAddChunkInfoPage(const FIoChunkId& ChunkId)
 	{
 		for (FJournalChunkInfoPage& Page : ChunkInfoPages)
 		{
@@ -357,13 +376,10 @@ namespace StorageServer
 			}
 		}
 
-		FJournalChunkInfoPage* ResPage = &ChunkInfoPages.Add_GetRef({ GMaxSectionEntries});
-		ResPage->SetFilePos(ReserveNextJournalFilePos(ResPage->GetPageMaxSize()));
-		bPagesModified = true;
-		return ResPage;
+		return AddChunkInfoPage();
 	}
 
-	FJournalChunkPage* FCacheJournalSectioned::FindChunkPage(const FIoChunkId& ChunkId)
+	FJournalChunkPage* FCacheJournalSectioned::FindOrAddChunkPage(const FIoChunkId& ChunkId)
 	{
 		FJournalChunkPage* LastPage = nullptr;
 		for (FJournalChunkPage& Page : ChunkPages)
@@ -379,18 +395,85 @@ namespace StorageServer
 			}
 		}
 
-		if (!LastPage)
+		return LastPage ? LastPage : AddChunkPage();
+	}
+
+	FJournalChunkInfoPage* FCacheJournalSectioned::AddChunkInfoPage()
+	{
+		FJournalChunkInfoPage* Page = &ChunkInfoPages.Add_GetRef({ JOURNAL_PAGE_SIZE, GMaxChunkInfoEntries });
+		Page->SetFilePos(AllocateJournalFilePos(Page->GetPageSize()));
+		bPagesModified = true;
+		return Page;
+	}
+
+	FJournalChunkPage* FCacheJournalSectioned::AddChunkPage()
+	{
+		FJournalChunkPage* Page = &ChunkPages.Add_GetRef({ JOURNAL_PAGE_SIZE, GMaxChunkEntries });
+		Page->SetFilePos(AllocateJournalFilePos(Page->GetPageSize()));
+		bPagesModified = true;
+		return Page;
+	}
+
+	bool FJournalPageBase::Flush(IFileHandle* JournalFile, TArray<uint8>& SerializationBuffer)
+	{
+		TRACE_CPUPROFILER_EVENT_SCOPE(FJournalPageBase::Flush);
+
+		if (!bDirty)
+		{
+			return true;
+		}
+
+		SerializationBuffer.SetNum(0, EAllowShrinking::No);
+		FMemoryWriter BodyAr(SerializationBuffer);
+		BodyAr << *this;
+
+		if (!ensureMsgf(SerializationBuffer.Num() + FJournalPageHeader::SerializedSize <= PageSize, TEXT("Page %u serialized to %i bytes > max size of %u"), PageType, SerializationBuffer.Num() + FJournalPageHeader::SerializedSize, PageSize))
+		{
+			return false;
+		}
+
+		FJournalPageHeader Descriptor;
+		Descriptor.Magic = PAGE_MAGIC;
+		Descriptor.Type = PageType;
+		Descriptor.PageSize = PageSize;
+		Descriptor.DataSize = SerializationBuffer.Num();
+
+		TArray<uint8, TInlineAllocator<64>> HeaderBuffer; // need to write payload size into header, so avoid using same buffer
+		TMemoryWriterBase HeaderAr(HeaderBuffer);
+		HeaderAr << Descriptor;
+
+		if (!ensure(HeaderBuffer.Num() <= FJournalPageHeader::SerializedSize))
+		{
+			return false;
+		}
+
+		if (!JournalFile->Seek(FilePos))
+		{
+			return false;
+		}
+
+		if (!JournalFile->Write(HeaderBuffer.GetData(), HeaderBuffer.Num()))
+		{
+			return false;
+		}
+
+		if (!JournalFile->Seek(FilePos + FJournalPageHeader::SerializedSize))
+		{
+			return false;
+		}
+
+		if (!JournalFile->Write(SerializationBuffer.GetData(), SerializationBuffer.Num()))
 		{
-			LastPage = &ChunkPages.Add_GetRef({ GMaxSectionEntries });
-			LastPage->SetFilePos(ReserveNextJournalFilePos(LastPage->GetPageMaxSize()));
-			bPagesModified = true;
+			return false;
 		}
 
-		return LastPage;
+		bDirty = false;
+		return true;
 	}
-	
-	FJournalChunkInfoPage::FJournalChunkInfoPage(int32 InMaxEntries)
-		: FJournalPageBase(InMaxEntries)
+
+	FJournalChunkInfoPage::FJournalChunkInfoPage(int64 InPageSize, int32 InMaxEntries)
+		: FJournalPageBase(EPageType::ChunkInfo, InPageSize)
+		, MaxEntryCount(InMaxEntries)
 	{
 	}
 
@@ -432,7 +515,12 @@ namespace StorageServer
 
 	void FJournalChunkInfoPage::Invalidate(const FIoChunkId& ChunkId)
 	{
-		bDirty = Entries.Remove(ChunkId) > 0;
+		const bool bRemoved = Entries.Remove(ChunkId) > 0;
+		if (bRemoved)
+		{
+			TRACE_COUNTER_DECREMENT(ZenJournalSectionedEntries);
+			bDirty = true;
+		}
 	}
 
 	void FJournalChunkInfoPage::InvalidateAll()
@@ -449,62 +537,15 @@ namespace StorageServer
 		}
 	}
 
-	bool FJournalChunkInfoPage::Flush(IFileHandle* Handle)
+	FArchive& operator<<(FArchive& Ar, FJournalChunkInfoPage& Page)
 	{
-		TRACE_CPUPROFILER_EVENT_SCOPE(FJournalChunkInfoPage::Flush);
-
-		if (!bDirty)
-		{
-			return true;
-		}
-
-		FJournalPageDescriptor Descriptor;
-		Descriptor.Magic = PAGE_MAGIC;
-		Descriptor.Type = EPageType::ChunkInfo;
-		Descriptor.MaxEntryCount = MaxEntryCount;
-
-		TArray<FChunkInfoEntry> OutputData;
-		OutputData.Reserve(Entries.Num());
-
-		for (decltype(Entries)::TIterator It = Entries.CreateIterator(); It; ++It)
-		{
-			OutputData.Add({ .ChunkId = It.Key(), .ChunkInfo = It.Value() });
-		}
-
-		Descriptor.Size = OutputData.Num() * OutputData.GetTypeSize();
-		Descriptor.MaxSize = MaxEntryCount * OutputData.GetTypeSize();
-
-		int64 CurrentPos = Handle->Tell();
-		Handle->Seek(FilePos);
-		if (!Handle->Write((uint8*)&Descriptor, sizeof(Descriptor)))
-		{
-			UE_LOG(LogCacheJournal, Error, TEXT("Failed to write chunk info page descriptor"));
-			return false;
-		}
-
-		if (!Handle->Write((uint8*)OutputData.GetData(), OutputData.Num() * OutputData.GetTypeSize()))
-		{
-			UE_LOG(LogCacheJournal, Error, TEXT("Failed to write chunk info entries"));
-			return false;
-		}
-		Handle->Seek(CurrentPos);
-
-		bDirty = false;
-		return true;
-	}
-
-	FJournalChunkInfoPage& FJournalChunkInfoPage::operator<<(const TArray<FChunkInfoEntry>& InEntries)
-	{
-		for (const FChunkInfoEntry& Entry : InEntries)
-		{
-			Entries.Add(Entry.ChunkId, Entry.ChunkInfo);
-		}
-
-		return *this;
+		Ar << Page.Entries;
+		return Ar;
 	}
 
-	FJournalChunkPage::FJournalChunkPage(int32 InMaxEntries)
-		: FJournalPageBase(InMaxEntries)
+	FJournalChunkPage::FJournalChunkPage(int64 InPageSize, int32 InMaxEntries)
+		: FJournalPageBase(EPageType::Chunk, InPageSize)
+		, MaxEntryCount(InMaxEntries)
 	{
 	}
 
@@ -529,6 +570,7 @@ namespace StorageServer
 		FIntervalMap& IntervalMap = ChunkMap.FindOrAdd(ChunkId);
 		IntervalMap.Add(Entry.GetChunkInterval(), Idx);
 
+		TRACE_COUNTER_INCREMENT(ZenJournalSectionedChunks);
 		bDirty = true;
 		return EJournalPageResult::Ok;
 	}
@@ -569,51 +611,6 @@ namespace StorageServer
 		}
 	}
 
-	bool FJournalChunkPage::Flush(IFileHandle* JournalFile)
-	{
-		TRACE_CPUPROFILER_EVENT_SCOPE(FJournalChunkPage::Flush);
-
-		if (!bDirty)
-		{
-			return true;
-		}
-
-		int64 CurrentPos = JournalFile->Tell();
-
-		FJournalPageDescriptor Descriptor;
-		Descriptor.Magic = PAGE_MAGIC;
-		Descriptor.Type = EPageType::Chunk;
-		Descriptor.MaxEntryCount = MaxEntryCount;
-		Descriptor.Size = Entries.Num() * Entries.GetTypeSize();
-		Descriptor.MaxSize = MaxEntryCount * Entries.GetTypeSize();
-		
-		if (bContainsInvalidEntries)
-		{
-			TArray<FJournalStoreEntry> ValidEntries;
-			ValidEntries.Reserve(Entries.Num());
-
-			const int32 EntriesCount = Entries.Num();
-			for (int32 Idx = 0; Idx < EntriesCount; ++Idx)
-			{
-				if (Entries[Idx].bValid)
-				{
-					ValidEntries.Add(Entries[Idx]);
-				}
-			}
-
-			Entries = MoveTemp(ValidEntries);
-			bContainsInvalidEntries = false;
-		}
-
-		JournalFile->Seek(FilePos);
-		JournalFile->Write((uint8*)&Descriptor, sizeof(Descriptor));
-		JournalFile->Write((uint8*)Entries.GetData(), Descriptor.Size);
-		JournalFile->Seek(CurrentPos);
-
-		bDirty = false;
-		return true;
-	}
-
 	void FJournalChunkPage::Invalidate(const FIoChunkId& ChunkId)
 	{
 		FIntervalMap* IntervalMap = ChunkMap.Find(ChunkId);
@@ -627,6 +624,7 @@ namespace StorageServer
 			Entries[It.Value()].bValid = false;
 		}
 
+		TRACE_COUNTER_DECREMENT(ZenJournalSectionedChunks);
 		ChunkMap.Remove(ChunkId);
 		bDirty = true;
 		bContainsInvalidEntries = true;
@@ -639,26 +637,51 @@ namespace StorageServer
 		bDirty = true;
 	}
 
-	FJournalChunkPage& FJournalChunkPage::operator<<(TArray<FJournalStoreEntry>&& InEntries)
+	FArchive& operator << (FArchive& Ar, FJournalChunkPage& Page)
 	{
-		Entries = MoveTemp(InEntries);
-		ChunkMap.Reset();
-
-		const int32 EntryCount = Entries.Num();
-		for (int32 Idx = 0; Idx < EntryCount; ++Idx)
+		bool bNeedsMapRebuild = Ar.IsLoading();
+		if (Page.bContainsInvalidEntries)
 		{
-			FJournalStoreEntry& StoreEntry = Entries[Idx];
-			if (!StoreEntry.bValid)
+			TArray<FJournalStoreEntry> ValidEntries;
+			ValidEntries.Reserve(Page.Entries.Num());
+
+			const int32 EntriesCount = Page.Entries.Num();
+			for (int32 Idx = 0; Idx < EntriesCount; ++Idx)
 			{
-				bContainsInvalidEntries = true;
-				continue;
+				if (Page.Entries[Idx].bValid)
+				{
+					ValidEntries.Add(Page.Entries[Idx]);
+				}
 			}
 
-			FIntervalMap& IntervalMap = ChunkMap.FindOrAdd(StoreEntry.ChunkId);
-			IntervalMap.Add(StoreEntry.Entry.GetChunkInterval(), Idx);
+			Page.Entries = MoveTemp(ValidEntries);
+			Page.bContainsInvalidEntries = false;
+
+			bNeedsMapRebuild = true;
+		}
+
+		Ar << Page.Entries;
+
+		if (bNeedsMapRebuild)
+		{
+			Page.ChunkMap.Empty();
+
+			const int32 EntryCount = Page.Entries.Num();
+			for (int32 Idx = 0; Idx < EntryCount; ++Idx)
+			{
+				FJournalStoreEntry& StoreEntry = Page.Entries[Idx];
+				if (!StoreEntry.bValid)
+				{
+					Page.bContainsInvalidEntries = true;
+					continue;
+				}
+
+				FJournalChunkPage::FIntervalMap& IntervalMap = Page.ChunkMap.FindOrAdd(StoreEntry.ChunkId);
+				IntervalMap.Add(StoreEntry.Entry.GetChunkInterval(), Idx);
+			}
 		}
 
-		return *this;
+		return Ar;
 	}
 
 }
diff --git a/Engine/Source/Runtime/StorageServerClient/Private/Cache/CacheJournalSectioned.h b/Engine/Source/Runtime/StorageServerClient/Private/Cache/CacheJournalSectioned.h
--- a/Engine/Source/Runtime/StorageServerClient/Private/Cache/CacheJournalSectioned.h
+++ b/Engine/Source/Runtime/StorageServerClient/Private/Cache/CacheJournalSectioned.h
@@ -4,33 +4,18 @@
 
 #include "CacheJournal.h"
 
+#if !UE_BUILD_SHIPPING
+
 #include "CoreMinimal.h"
 #include "Serialization/MemoryReader.h"
-
-#if !UE_BUILD_SHIPPING
+#include "Tasks/Task.h"
 
 namespace StorageServer
 {
-	inline constexpr uint32 GMaxSectionEntries = 1000;
-	inline constexpr uint32 JOURNAL_PAGE_SIZE = 128 * 1024;
-
 	enum class EPageType : uint8
 	{
-		Chunk,
-		ChunkInfo
-	};
-
-	struct FJournalChunkInfoEntry
-	{
-		FIoChunkId ChunkId;
-		FCacheChunkInfo ChunkInfo;
-	};
-
-	struct FPageDescriptor
-	{
-		EPageType Type;
-		int32 NumEntries;
-		uint8* Data;
+		Chunk = 1,
+		ChunkInfo = 2
 	};
 
 	struct FJournalStoreEntry
@@ -43,6 +28,14 @@ namespace StorageServer
 		{
 			return ChunkId < Other.ChunkId;
 		}
+
+		friend FArchive& operator<<(FArchive& Ar, FJournalStoreEntry& StoreEntry)
+		{
+			Ar << StoreEntry.ChunkId;
+			Ar << StoreEntry.Entry;
+			Ar << StoreEntry.bValid;
+			return Ar;
+		}
 	};
 
 	enum class EJournalPageResult
@@ -55,25 +48,47 @@ namespace StorageServer
 
 	struct FJournalHeader
 	{
+		uint32 Magic;
 		uint32 Version;
-		int32 NumberOfChunkPages;
-		int32 NumberOfChunkInfoPages;
+		uint32 PageCount;
+
+		static constexpr int32 SerializedSize = 32; // actual size is 12 bytes, but we reserve some for future use
+
+		friend FArchive& operator << (FArchive& Ar, FJournalHeader& Header)
+		{
+			Ar << Header.Magic;
+			Ar << Header.Version;
+			Ar << Header.PageCount;
+			return Ar;
+		}
 	};
 
-	struct FJournalPageDescriptor
+	struct FJournalPageHeader
 	{
 		uint32 Magic;
+		uint32 PageSize; // total size of page in bytes, should be used to calculate offset of next page
+		uint32 DataSize; // size of data in current page minus header, can be less than page size if page is not full
 		EPageType Type;
-		uint32 MaxEntryCount;
-		uint32 MaxSize;
-		uint32 Size;
+
+		static constexpr int32 SerializedSize = 32; // actual size is 13 bytes, but we reserve some for future use
+
+		friend FArchive& operator << (FArchive& Ar, FJournalPageHeader& Desc)
+		{
+			Ar << Desc.Magic;
+			Ar << Desc.PageSize;
+			Ar << Desc.DataSize;
+			Ar << Desc.Type;
+			return Ar;
+		}
 	};
 
 	class FJournalPageBase
 	{
 	public:
-		FJournalPageBase(int32 InMaxEntryCount)
-			: MaxEntryCount(InMaxEntryCount)
+		FJournalPageBase(EPageType InPageType, int64 InPageSize)
+			: FilePos(-1)
+			, PageSize(InPageSize)
+			, PageType(InPageType)
 		{
 		}
 
@@ -84,101 +99,111 @@ namespace StorageServer
 		[[nodiscard]] virtual bool Contains(const FIoChunkId&) const = 0;
 		[[nodiscard]] virtual bool IsFull() const = 0;
 
-		[[nodiscard]] virtual int32 GetPageMaxSize() const = 0;
 		void SetFilePos(int64 FileCursor)
 		{
 			FilePos = FileCursor;
 		}
 
-		virtual bool Flush(IFileHandle* JournalFile) = 0;
+		int64 GetFilePos() const
+		{
+			return FilePos;
+		}
+
+		int64 GetPageSize() const
+		{
+			return PageSize;
+		}
+
+		bool Flush(IFileHandle* JournalFile, TArray<uint8>& SerializationBuffer);
+
+		friend FArchive& operator<<(FArchive& Ar, FJournalPageBase& Page)
+		{
+			Page.Serialize(Ar);
+			return Ar;
+		}
 
 	protected:
 		int64 FilePos;
-		int32 MaxEntryCount;
+		const int64 PageSize;
+		const EPageType PageType;
+		bool bDirty = false;
+
+		virtual void Serialize(FArchive& Ar) = 0;
 	};
 
 	class FJournalChunkInfoPage : public FJournalPageBase
 	{
 	public:
-		FJournalChunkInfoPage(int32 InMaxEntries);
+		FJournalChunkInfoPage(int64 InPageSize, int32 InMaxEntries);
 
 		EJournalPageResult Add(const FIoChunkId& ChunkId, const FCacheChunkInfo& Entry);
 		bool GetEntry(const FIoChunkId& ChunkId, FCacheChunkInfo& Info);
-		bool Contains(const FIoChunkId& ChunkId) const override;
-		bool IsFull() const override 
+		virtual bool Contains(const FIoChunkId& ChunkId) const override;
+		virtual bool IsFull() const override 
 		{
 			return Entries.Num() >= MaxEntryCount;
 		}
+		int32 GetEntryCount() const
+		{
+			return Entries.Num();
+		}
 
 		void Invalidate(const FIoChunkId& ChunkId);
 		void InvalidateAll();
 
 		void IterateChunkIds(TFunctionRef<void(const FIoChunkId& ChunkId, const FCacheChunkInfo& ChunkInfo)> Callback);
 
-		int32 GetPageMaxSize() const override
-		{
-			return MaxEntryCount * sizeof(FChunkInfoEntry);
-		}
-		
-		bool Flush(IFileHandle* JournalFile) override;
-
-		struct FChunkInfoEntry
-		{
-			FIoChunkId ChunkId;
-			FCacheChunkInfo ChunkInfo;
-		};
-
-		FJournalChunkInfoPage& operator<<(const TArray<FChunkInfoEntry>& InEntries);
+		friend FArchive& operator<<(FArchive& Ar, FJournalChunkInfoPage& Page);
 
 	private:
-		bool bDirty = false;
-
+		int32 MaxEntryCount;
 		TMap<FIoChunkId, FCacheChunkInfo> Entries;
+
+		virtual void Serialize(FArchive& Ar) override
+		{
+			Ar << *this;
+		}
 	};
 
 	class FJournalChunkPage : public FJournalPageBase
 	{
 	public:
-		FJournalChunkPage(int32 InMaxEntries);
+		FJournalChunkPage(int64 InPageSize, int32 InMaxEntries);
 
 		EJournalPageResult Add(const FIoChunkId& ChunkId, const FCacheEntry& Entry);
 		bool GetEntry(const FIoChunkId& ChunkId, const uint64 ChunkOffset, const uint64 ChunkSize, FCacheEntry& OutEntry);
-		bool Contains(const FIoChunkId& ChunkId) const override
+		virtual bool Contains(const FIoChunkId& ChunkId) const override
 		{
 			return ChunkMap.Contains(ChunkId);
 		}
 
 		void IterateCacheEntriesForChunkId(const FIoChunkId& ChunkId, TFunctionRef<void(const FCacheEntry& Entry)> Callback);
 		void IterateCacheEntries(TFunctionRef<void(const FIoChunkId& ChunkId, const FCacheEntry& Entry)> Callback);
-
-		bool IsFull() const override
+		virtual bool IsFull() const override
 		{
 			return Entries.Num() == MaxEntryCount;
 		}
-
-		bool Flush(IFileHandle* JournalFile) override;
-
-		int32 GetPageMaxSize() const override
-		{
-			return MaxEntryCount * Entries.GetTypeSize();
-		}
-
-		int32 GetEntrySize() const
+		int32 GetEntryCount() const
 		{
-			return Entries.GetTypeSize();
+			return Entries.Num();
 		}
 
 		void Invalidate(const FIoChunkId& ChunkId);
 		void InvalidateAll();
 
-		FJournalChunkPage& operator<<(TArray<FJournalStoreEntry>&& InEntries);
+		friend FArchive& operator<<(FArchive& Ar, FJournalChunkPage& Page);
 
 	private:
 		using FIntervalMap = TMap<TInterval<uint64>, int32>;
 		TMap<FIoChunkId, FIntervalMap> ChunkMap;
 		TArray<FJournalStoreEntry> Entries;
-		bool bDirty = false;
+		int32 MaxEntryCount;
 		bool bContainsInvalidEntries = false;
+
+		virtual void Serialize(FArchive& Ar) override
+		{
+			Ar << *this;
+		}
 	};
 
 	class FCacheJournalSectioned : public ICacheJournal
@@ -188,7 +213,7 @@ namespace StorageServer
 		FCacheJournalSectioned(FStringView JournalPath);
 		virtual ~FCacheJournalSectioned() override = default;
 
-		virtual void Flush() override;
+		virtual void Flush(bool bImmediate) override;
 		virtual void InvalidateAll() override;
 		virtual void Invalidate(const FIoChunkId& ChunkId) override;
 
@@ -203,33 +228,41 @@ namespace StorageServer
 		virtual void IterateCacheEntries(TFunctionRef<void(const FIoChunkId& ChunkId, const FCacheEntry& Entry)> Callback) override;
 
 	private:
-		void LoadCacheJournal(FStringView JournalPath);
+		bool LoadCacheJournal();
 
-		FJournalChunkInfoPage* FindAddChunkInfoPage(const FIoChunkId& ChunkId);
-		FJournalChunkPage* FindChunkPage(const FIoChunkId& ChunkId);
+		FJournalChunkInfoPage* FindOrAddChunkInfoPage(const FIoChunkId& ChunkId);
+		FJournalChunkPage* FindOrAddChunkPage(const FIoChunkId& ChunkId);
 
-		int32 ReserveNextJournalFilePos(int32 PageSize)
+		FJournalChunkInfoPage* AddChunkInfoPage();
+		FJournalChunkPage* AddChunkPage();
+
+		int32 AllocateJournalFilePos(int32 PageSize)
 		{
 			int32 CurrentJournalPos = NextAvailableJournalPos;
-			NextAvailableJournalPos += sizeof(FJournalPageDescriptor) + PageSize;
+			NextAvailableJournalPos += PageSize;
 			return CurrentJournalPos;
 		}
 
+		void ResetJournalFilePos()
+		{
+			NextAvailableJournalPos = FJournalHeader::SerializedSize;
+		}
+
 	private:
 		FString JournalFilePath;
-		uint32 JournalSectionCursor = sizeof(FJournalHeader);
 		TUniquePtr<IFileHandle> JournalFileHandle;
-		uint32 NextAvailableJournalPos = sizeof(FJournalHeader);
+		uint32 NextAvailableJournalPos = FJournalHeader::SerializedSize;
+		UE::Tasks::TTask<void> FlushTask;
 
 		bool bPagesModified = false;
 
-		using FEntryIntervalMap = TMap<TInterval<uint64>, FCacheEntry>;
-
 		TArray<FJournalChunkPage> ChunkPages;
 		TArray<FJournalChunkInfoPage> ChunkInfoPages;
 		FJournalChunkPage* CurrentChunkPage = nullptr;
 
 		FCriticalSection CriticalSection;
+
+		void FlushImmediate();
 	};
 
 }
diff --git a/Engine/Source/Runtime/StorageServerClient/Private/Cache/CacheJournalSimple.cpp b/Engine/Source/Runtime/StorageServerClient/Private/Cache/CacheJournalSimple.cpp
--- a/Engine/Source/Runtime/StorageServerClient/Private/Cache/CacheJournalSimple.cpp
+++ b/Engine/Source/Runtime/StorageServerClient/Private/Cache/CacheJournalSimple.cpp
@@ -43,61 +43,28 @@ namespace StorageServer
 	{
 	}
 
-	void FCacheJournalSimple::Flush()
+	void FCacheJournalSimple::Flush(bool bImmediate)
 	{
 		TRACE_CPUPROFILER_EVENT_SCOPE(FCacheJournalSimple::Flush);
 
-		if (!bDirty) // no need to sync this
+		if (bImmediate)
 		{
-			return;
+			FlushImmediate();
 		}
-
-		FScopeLock Lock(&FlushTaskLock);
-
-		if (!FlushTask.IsCompleted())
+		else
 		{
-			return;
-		}
-
-		FlushTask = UE::Tasks::Launch(TEXT("CacheJournalFlush"), [this]()
-		{
-			TRACE_CPUPROFILER_EVENT_SCOPE(FCacheJournalSimple::Flush::Task);
-
-			FScopeLock Lock(&DataLock);
-
-			if (!bDirty)
-			{
-				return;
-			}
-
-			// TODO avoid writing whole file every time
-			FArchive* Writer = IFileManager::Get().CreateFileWriter(*FileName);
-			if (!Writer)
+			if (!bDirty) // no need to sync this
 			{
 				return;
 			}
 
-			FArchive& Ar = *Writer;
+			FScopeLock Lock(&FlushTaskLock);
 
-			// As we only care about chunkid's that are actually cached there is no need to store all chunk id's
-			decltype(ChunkInfos) OnlyCachedChunksInfo;
-
-			for (const TTuple<FIoChunkId, FPerChunkCacheEntries>& Pair : ChunkCacheEntries)
+			if (FlushTask.IsCompleted())
 			{
-				FCacheChunkInfo* Info = nullptr;
-				if ((Info = ChunkInfos.Find(Pair.Key)) != nullptr)
-				{
-					OnlyCachedChunksInfo.Add(Pair.Key, *Info);
-				}
+				FlushTask = UE::Tasks::Launch(TEXT("CacheJournalFlush"), [this]() {FlushImmediate();});
 			}
-
-			Ar << OnlyCachedChunksInfo;
-			Ar << ChunkCacheEntries;
-
-			Writer->Flush();
-			delete Writer;
-			bDirty = false;
-		});
+		}
 	}
 
 	void FCacheJournalSimple::InvalidateAll()
@@ -134,6 +101,7 @@ namespace StorageServer
 		FCacheChunkInfo* ChunkInfo = ChunkInfos.Find(ChunkId);
 		if (ChunkInfo)
 		{
+			bDirty = true; 
 			return ChunkInfo->SetChunkInfo(OptModTag, OptRawSize, OptRawBlockSize);
 		}
 
@@ -189,7 +157,7 @@ namespace StorageServer
 
 		if (bNeedsFlush)
 		{
-			Flush();
+			Flush(false);
 		}
 
 		return true;
@@ -258,6 +226,46 @@ namespace StorageServer
 			}
 		}
 	}
+
+	void FCacheJournalSimple::FlushImmediate()
+	{
+		TRACE_CPUPROFILER_EVENT_SCOPE(FCacheJournalSimple::FlushImmediate);
+
+		FScopeLock Lock(&DataLock);
+
+		if (!bDirty)
+		{
+			return;
+		}
+
+		// TODO avoid writing whole file every time
+		FArchive* Writer = IFileManager::Get().CreateFileWriter(*FileName);
+		if (!Writer)
+		{
+			return;
+		}
+
+		FArchive& Ar = *Writer;
+
+		// As we only care about chunkid's that are actually cached there is no need to store all chunk id's
+		decltype(ChunkInfos) OnlyCachedChunksInfo;
+
+		for (const TTuple<FIoChunkId, FPerChunkCacheEntries>& Pair : ChunkCacheEntries)
+		{
+			FCacheChunkInfo* Info = nullptr;
+			if ((Info = ChunkInfos.Find(Pair.Key)) != nullptr)
+			{
+				OnlyCachedChunksInfo.Add(Pair.Key, *Info);
+			}
+		}
+
+		Ar << OnlyCachedChunksInfo;
+		Ar << ChunkCacheEntries;
+
+		Writer->Flush();
+		delete Writer;
+		bDirty = false;
+	}
 }
 
 #endif
diff --git a/Engine/Source/Runtime/StorageServerClient/Private/Cache/CacheJournalSimple.h b/Engine/Source/Runtime/StorageServerClient/Private/Cache/CacheJournalSimple.h
--- a/Engine/Source/Runtime/StorageServerClient/Private/Cache/CacheJournalSimple.h
+++ b/Engine/Source/Runtime/StorageServerClient/Private/Cache/CacheJournalSimple.h
@@ -19,7 +19,7 @@ namespace StorageServer
 		FCacheJournalSimple(const TCHAR* FileName, const uint64 FlushAtWriteCount);
 		virtual ~FCacheJournalSimple() override;
 
-		virtual void Flush() override;
+		virtual void Flush(bool bImmediate) override;
 
 		virtual void InvalidateAll() override;
 		virtual void Invalidate(const FIoChunkId& ChunkId) override;
@@ -68,6 +68,8 @@ namespace StorageServer
 		typedef TMap<TInterval<uint64>, FCacheEntry> FPerChunkCacheEntries;
 		TMap<FIoChunkId, FCacheChunkInfo> ChunkInfos;
 		TMap<FIoChunkId, FPerChunkCacheEntries> ChunkCacheEntries;
+
+		void FlushImmediate();
 	};
 }
 
diff --git a/Engine/Source/Runtime/StorageServerClient/Private/Cache/CacheStrategy.h b/Engine/Source/Runtime/StorageServerClient/Private/Cache/CacheStrategy.h
--- a/Engine/Source/Runtime/StorageServerClient/Private/Cache/CacheStrategy.h
+++ b/Engine/Source/Runtime/StorageServerClient/Private/Cache/CacheStrategy.h
@@ -94,7 +94,7 @@ namespace StorageServer
 
 		virtual void Flush()
 		{
-			Journal->Flush();
+			Journal->Flush(false);
 			Storage->Flush();
 		}
 
