- Added support for IterateDirectory and friends to act like SandboxWrapper for running on remote and local (only when using hybrid cooked editor, or otherwise enabling s.PreferLocalForNonAssets)

diff --git a/Engine/Source/Runtime/StorageServerClient/Private/StorageServerPlatformFile.cpp b/Engine/Source/Runtime/StorageServerClient/Private/StorageServerPlatformFile.cpp
--- a/Engine/Source/Runtime/StorageServerClient/Private/StorageServerPlatformFile.cpp
+++ b/Engine/Source/Runtime/StorageServerClient/Private/StorageServerPlatformFile.cpp
@@ -1343,32 +1343,44 @@ IFileHandle* FStorageServerPlatformFile::OpenRead(const TCHAR* Filename, bool bA
 bool FStorageServerPlatformFile::IterateDirectory(const TCHAR* Directory, IPlatformFile::FDirectoryVisitor& Visitor)
 {
 	TStringBuilder<1024> StorageServerDirectory;
-	bool bResult = false;
-	if (MakeStorageServerPath(Directory, StorageServerDirectory) && ServerToc.DirectoryExists(*StorageServerDirectory))
+	bool bShouldContinue = true;
+	bool bIterateOnServer = MakeStorageServerPath(Directory, StorageServerDirectory) && ServerToc.DirectoryExists(*StorageServerDirectory);
+
+	if (bIterateOnServer)
 	{
-		bResult |= ServerToc.IterateDirectory(*StorageServerDirectory, [this, &Visitor](const FIoChunkId& FileChunkId, const TCHAR* FilenameOrDirectory, int64 RawSize)
-		{
-			TStringBuilder<1024> LocalPath;
-			bool bConverted = MakeLocalPath(FilenameOrDirectory, LocalPath);
-			check(bConverted);
-			const bool bDirectory = !FileChunkId.IsValid();
-			return Visitor.CallShouldVisitAndVisit(*LocalPath, bDirectory);
-		});
+		// a false return here means the visitor returned false, which means to stop iterating on more files
+		// it can also mean the directory doesn't exist, but we just checked that above, so false only means to stop iterating
+		bShouldContinue = ServerToc.IterateDirectory(*StorageServerDirectory, [this, &Visitor](const FIoChunkId& FileChunkId, const TCHAR* FilenameOrDirectory, int64 RawSize)
+			{
+				TStringBuilder<1024> LocalPath;
+				bool bConverted = MakeLocalPath(FilenameOrDirectory, LocalPath);
+				check(bConverted);
+				const bool bDirectory = !FileChunkId.IsValid();
+				return Visitor.CallShouldVisitAndVisit(*LocalPath, bDirectory);
+			});
 	}
-	else if (LowerLevel)
+
+	// run a local iterate in 2 cases:
+	//   1. we prefer local and we were not told to stop iterating above (matches FSandboxPlatformFile)
+	//   2. we don't prefer local, and the directory doesn't exist on the server (this is the existing behavior, before we added GPreferLocalForNonAssets)
+	// this code assumes (as does Sandbox seemingly) that the visitors can handle the same path being returned twice
+	if ((GPreferLocalForNonAssets && bShouldContinue) || (!GPreferLocalForNonAssets && !bIterateOnServer))
 	{
-		bResult |= LowerLevel->IterateDirectory(Directory, Visitor);
+		bShouldContinue = LowerLevel->IterateDirectory(Directory, Visitor);
 	}
-	return bResult;
+	return bShouldContinue;
 }
 
 bool FStorageServerPlatformFile::IterateDirectoryRecursively(const TCHAR* Directory, IPlatformFile::FDirectoryVisitor& Visitor)
 {
 	TStringBuilder<1024> StorageServerDirectory;
-	bool bResult = false;
-	if (MakeStorageServerPath(Directory, StorageServerDirectory) && ServerToc.DirectoryExists(*StorageServerDirectory))
+	bool bShouldContinue = true;
+	bool bIterateOnServer = MakeStorageServerPath(Directory, StorageServerDirectory) && ServerToc.DirectoryExists(*StorageServerDirectory);
+
+	// see IterateDirectory() comments for the flow logic
+	if (bIterateOnServer)
 	{
-		bResult |= ServerToc.IterateDirectoryRecursively(*StorageServerDirectory, [this, &Visitor](const FIoChunkId& FileChunkId, const TCHAR* FilenameOrDirectory, int64 RawSize)
+		bShouldContinue = ServerToc.IterateDirectoryRecursively(*StorageServerDirectory, [this, &Visitor](const FIoChunkId& FileChunkId, const TCHAR* FilenameOrDirectory, int64 RawSize)
 		{
 			TStringBuilder<1024> LocalPath;
 			bool bConverted = MakeLocalPath(FilenameOrDirectory, LocalPath);
@@ -1377,21 +1389,23 @@ bool FStorageServerPlatformFile::IterateDirectoryRecursively(const TCHAR* Direct
 			return Visitor.CallShouldVisitAndVisit(*LocalPath, bDirectory);
 		});
 	}
-	else
+	if ((GPreferLocalForNonAssets && bShouldContinue) || (!GPreferLocalForNonAssets && !bIterateOnServer))
 	{
-		bResult |= LowerLevel->IterateDirectoryRecursively(Directory, Visitor);
+		bShouldContinue = LowerLevel->IterateDirectoryRecursively(Directory, Visitor);
 	}
-
-	return bResult;
+	return bShouldContinue;
 }
 
 bool FStorageServerPlatformFile::IterateDirectoryStat(const TCHAR* Directory, FDirectoryStatVisitor& Visitor)
 {
 	TStringBuilder<1024> StorageServerDirectory;
-	bool bResult = false;
-	if (MakeStorageServerPath(Directory, StorageServerDirectory) && ServerToc.DirectoryExists(*StorageServerDirectory))
+	bool bShouldContinue = true;
+	bool bIterateOnServer = MakeStorageServerPath(Directory, StorageServerDirectory) && ServerToc.DirectoryExists(*StorageServerDirectory);
+
+	// see IterateDirectory() comments for the flow logic
+	if (bIterateOnServer)
 	{
-		bResult |= ServerToc.IterateDirectory(*StorageServerDirectory, [this, &Visitor](const FIoChunkId& FileChunkId, const TCHAR* ServerFilenameOrDirectory, int64 RawSize)
+		bShouldContinue = ServerToc.IterateDirectory(*StorageServerDirectory, [this, &Visitor](const FIoChunkId& FileChunkId, const TCHAR* ServerFilenameOrDirectory, int64 RawSize)
 		{
 			TStringBuilder<1024> LocalPath;
 			bool bConverted = MakeLocalPath(ServerFilenameOrDirectory, LocalPath);
@@ -1421,11 +1435,11 @@ bool FStorageServerPlatformFile::IterateDirectoryStat(const TCHAR* Directory, FD
 			return Visitor.CallShouldVisitAndVisit(*LocalPath, FileStatData);
 		});
 	}
-	else if (LowerLevel)
+	if ((GPreferLocalForNonAssets && bShouldContinue) || (!GPreferLocalForNonAssets && !bIterateOnServer))
 	{
-		bResult |= LowerLevel->IterateDirectoryStat(Directory, Visitor);
+		bShouldContinue = LowerLevel->IterateDirectoryStat(Directory, Visitor);
 	}
-	return bResult;
+	return bShouldContinue;
 }
 
 FOpenMappedResult FStorageServerPlatformFile::OpenMappedEx(const TCHAR* Filename, EOpenReadFlags OpenOptions, int64 MaximumSize)
