Adding zen sectioned journal flush and page hashing to ensure resilience to partial file writes

diff --git a/Engine/Source/Runtime/StorageServerClient/Private/Cache/CacheJournalSectioned.cpp b/Engine/Source/Runtime/StorageServerClient/Private/Cache/CacheJournalSectioned.cpp
--- a/Engine/Source/Runtime/StorageServerClient/Private/Cache/CacheJournalSectioned.cpp
+++ b/Engine/Source/Runtime/StorageServerClient/Private/Cache/CacheJournalSectioned.cpp
@@ -18,7 +18,7 @@ namespace StorageServer
 {
 	static inline constexpr uint64 FILE_MAGIC = 0x5a45'4e43; // ZENC
 	static inline constexpr uint64 PAGE_MAGIC = 0x5041'4745; // PAGE
-	static inline constexpr uint32 JOURNAL_VERSION = 0x02;
+	static inline constexpr uint32 JOURNAL_VERSION = 0x03;
 
 	// TODO hardcoded constants of max entries are fragile, can we compute it from page size directly?
 	static inline constexpr uint32 GMaxChunkInfoEntries = 2950; // approx 44 bytes per entry
@@ -60,6 +60,8 @@ namespace StorageServer
 
 		TArray<uint8> SerializationBuffer;
 
+		bool bNeedsFlush = false;
+
 		if (bPagesModified)
 		{
 			FJournalHeader Header;
@@ -87,16 +89,22 @@ namespace StorageServer
 			}
 
 			bPagesModified = false;
+			bNeedsFlush = true;
 		}
 
 		for (FJournalChunkPage& Page : ChunkPages)
 		{
-			Page.Flush(JournalFileHandle.Get(), SerializationBuffer);
+			bNeedsFlush |= Page.Flush(JournalFileHandle.Get(), SerializationBuffer);
 		}
 
 		for (FJournalChunkInfoPage& Page : ChunkInfoPages)
 		{
-			Page.Flush(JournalFileHandle.Get(), SerializationBuffer);
+			bNeedsFlush |= Page.Flush(JournalFileHandle.Get(), SerializationBuffer);
+		}
+
+		if (bNeedsFlush)
+		{
+			JournalFileHandle->Flush(true);
 		}
 	}
 
@@ -317,6 +325,15 @@ namespace StorageServer
 				return false;
 			}
 
+			const FIoHash DataHash = FIoHash::HashBuffer(SerializationBuffer.GetData(), SerializationBuffer.Num());
+
+			if (DataHash != PageHeader.DataHash)
+			{
+				UE_LOG(LogCacheJournal, Warning, TEXT("Page has invalid data hash, the journal is corrupt."));
+				InvalidateAll();
+				return false;
+			}
+
 			FJournalPageBase* Page = nullptr;
 			switch (PageHeader.Type)
 			{
@@ -420,7 +437,7 @@ namespace StorageServer
 
 		if (!bDirty)
 		{
-			return true;
+			return false; // don't require flush if nothing changed
 		}
 
 		SerializationBuffer.SetNum(0, EAllowShrinking::No);
@@ -437,8 +454,9 @@ namespace StorageServer
 		Descriptor.Type = PageType;
 		Descriptor.PageSize = PageSize;
 		Descriptor.DataSize = SerializationBuffer.Num();
+		Descriptor.DataHash = FIoHash::HashBuffer(SerializationBuffer.GetData(), SerializationBuffer.Num());
 
-		TArray<uint8, TInlineAllocator<64>> HeaderBuffer; // need to write payload size into header, so avoid using same buffer
+		TArray<uint8, TInlineAllocator<128>> HeaderBuffer; // need to write payload size into header, so avoid using same buffer
 		TMemoryWriterBase HeaderAr(HeaderBuffer);
 		HeaderAr << Descriptor;
 
diff --git a/Engine/Source/Runtime/StorageServerClient/Private/Cache/CacheJournalSectioned.h b/Engine/Source/Runtime/StorageServerClient/Private/Cache/CacheJournalSectioned.h
--- a/Engine/Source/Runtime/StorageServerClient/Private/Cache/CacheJournalSectioned.h
+++ b/Engine/Source/Runtime/StorageServerClient/Private/Cache/CacheJournalSectioned.h
@@ -68,15 +68,17 @@ namespace StorageServer
 		uint32 Magic;
 		uint32 PageSize; // total size of page in bytes, should be used to calculate offset of next page
 		uint32 DataSize; // size of data in current page minus header, can be less than page size if page is not full
+		FIoHash DataHash; // hash of data in current page
 		EPageType Type;
 
-		static constexpr int32 SerializedSize = 32; // actual size is 13 bytes, but we reserve some for future use
+		static constexpr int32 SerializedSize = 64; // actual size is 33 bytes, but we reserve some for future use
 
 		friend FArchive& operator << (FArchive& Ar, FJournalPageHeader& Desc)
 		{
 			Ar << Desc.Magic;
 			Ar << Desc.PageSize;
 			Ar << Desc.DataSize;
+			Ar << Desc.DataHash;
 			Ar << Desc.Type;
 			return Ar;
 		}
diff --git a/Engine/Source/Runtime/StorageServerClient/Private/StorageServerConnection.cpp b/Engine/Source/Runtime/StorageServerClient/Private/StorageServerConnection.cpp
--- a/Engine/Source/Runtime/StorageServerClient/Private/StorageServerConnection.cpp
+++ b/Engine/Source/Runtime/StorageServerClient/Private/StorageServerConnection.cpp
@@ -352,7 +352,7 @@ FIoStatus FStorageServerConnection::ReadChunkBatchRequest(const TArray<FChunkBat
 		{
 			if (RequestedChunkIds.Contains(Chunk.ChunkId))
 			{
-				UE_LOG(LogStorageServerConnection, Fatal, TEXT("Multiple requests to the same chunkid are not supported"));
+				UE_LOG(LogStorageServerConnection, Fatal, TEXT("Multiple requests to the same chunkid are not supported (%s is already present)"), *LexToString(Chunk.ChunkId));
 				return FIoStatus::Unknown;
 			}
 
