[StorageServerClient]
* Fixed msvc analyzer compile warning

diff --git a/Engine/Source/Runtime/StorageServerClient/Private/StorageServerConnection.cpp b/Engine/Source/Runtime/StorageServerClient/Private/StorageServerConnection.cpp
--- a/Engine/Source/Runtime/StorageServerClient/Private/StorageServerConnection.cpp
+++ b/Engine/Source/Runtime/StorageServerClient/Private/StorageServerConnection.cpp
@@ -815,7 +815,12 @@ void FStorageServerConnection::GetDefaultCacheConfiguration(FCacheConfiguration&
 {
 	const TCHAR* CmdLine = FCommandLine::Get();
 
-	OutConfiguration.bEnable = PLATFORM_SUPPORTS_STORAGE_SERVER_CACHE && (FCString::Strstr(CmdLine, TEXT("-ZenEnableCache")) != nullptr);
+#if PLATFORM_SUPPORTS_STORAGE_SERVER_CACHE
+	OutConfiguration.bEnable = (FCString::Strstr(CmdLine, TEXT("-ZenEnableCache")) != nullptr);
+#else
+	OutConfiguration.bEnable = false;
+#endif
+
 	if (!OutConfiguration.bEnable)
 	{
 		return;
