Enable zen streaming cache by default on some platforms

diff --git a/Engine/Source/Runtime/StorageServerClient/Private/StorageServerConnection.cpp b/Engine/Source/Runtime/StorageServerClient/Private/StorageServerConnection.cpp
--- a/Engine/Source/Runtime/StorageServerClient/Private/StorageServerConnection.cpp
+++ b/Engine/Source/Runtime/StorageServerClient/Private/StorageServerConnection.cpp
@@ -815,11 +815,20 @@ void FStorageServerConnection::GetDefaultCacheConfiguration(FCacheConfiguration&
 {
 	const TCHAR* CmdLine = FCommandLine::Get();
 
-#if PLATFORM_SUPPORTS_STORAGE_SERVER_CACHE
-	OutConfiguration.bEnable = (FCString::Strstr(CmdLine, TEXT("-ZenEnableCache")) != nullptr);
-#else
-	OutConfiguration.bEnable = false;
-#endif
+	const bool bPlatformSupportsCaching = PLATFORM_SUPPORTS_STORAGE_SERVER_CACHE;
+	const bool bPlatformEnablesCachingByDefault = PLATFORM_ENABLES_STORAGE_SERVER_CACHE_BY_DEFAULT;
+
+	OutConfiguration.bEnable = bPlatformSupportsCaching && bPlatformEnablesCachingByDefault;
+
+	// always check it first to ensure we disable cache if cmd arg is provided
+	if ((FCString::Strstr(CmdLine, TEXT("-ZenDisableCache")) != nullptr))
+	{
+		OutConfiguration.bEnable = false;
+	}
+	else if (bPlatformSupportsCaching && (FCString::Strstr(CmdLine, TEXT("-ZenEnableCache")) != nullptr))
+	{
+		OutConfiguration.bEnable = true;
+	}
 
 	if (!OutConfiguration.bEnable)
 	{
@@ -838,8 +847,11 @@ void FStorageServerConnection::GetDefaultCacheConfiguration(FCacheConfiguration&
 	OutConfiguration.FlushInterval = 10.0f;
 	OutConfiguration.FlushEveryNEntries = 0;
 
-	OutConfiguration.bUseSectionedJournal = FCString::Strstr(CmdLine, TEXT("-ZenUseSectionedJournal")) != nullptr;
-	OutConfiguration.bUseMemoryMappedStorage = FPlatformProperties::SupportsMemoryMappedFiles() && (FCString::Strstr(CmdLine, TEXT("-ZenMmapedStorage")) != nullptr);
+	// use sectioned journal by default
+	OutConfiguration.bUseSectionedJournal = FCString::Strstr(CmdLine, TEXT("-ZenUseSimpleJournal")) == nullptr;
+
+	// use mmapp'ed storage if available and not explicitely disabled 
+	OutConfiguration.bUseMemoryMappedStorage = FPlatformProperties::SupportsMemoryMappedFiles() && (FCString::Strstr(CmdLine, TEXT("-ZenNoMmappedStorage")) == nullptr);
 }
 
 #if !HAS_STORAGE_SERVER_PLATFORM_CACHE_STRATEGY
diff --git a/Engine/Source/Runtime/StorageServerClient/Private/StorageServerConnection.h b/Engine/Source/Runtime/StorageServerClient/Private/StorageServerConnection.h
--- a/Engine/Source/Runtime/StorageServerClient/Private/StorageServerConnection.h
+++ b/Engine/Source/Runtime/StorageServerClient/Private/StorageServerConnection.h
@@ -21,6 +21,10 @@
 	#define PLATFORM_SUPPORTS_STORAGE_SERVER_CACHE (PLATFORM_WINDOWS || PLATFORM_ANDROID)
 #endif
 
+#ifndef PLATFORM_ENABLES_STORAGE_SERVER_CACHE_BY_DEFAULT
+	#define PLATFORM_ENABLES_STORAGE_SERVER_CACHE_BY_DEFAULT 0
+#endif
+
 #ifndef HAS_STORAGE_SERVER_PLATFORM_CACHE_STRATEGY
 	#define HAS_STORAGE_SERVER_PLATFORM_CACHE_STRATEGY 0
 #endif
