Enable zen streaming caching on iOS

#rnx

diff --git a/Engine/Source/Runtime/Core/Private/IOS/IOSPlatformMisc.cpp b/Engine/Source/Runtime/Core/Private/IOS/IOSPlatformMisc.cpp
--- a/Engine/Source/Runtime/Core/Private/IOS/IOSPlatformMisc.cpp
+++ b/Engine/Source/Runtime/Core/Private/IOS/IOSPlatformMisc.cpp
@@ -2184,6 +2184,12 @@ bool FIOSPlatformMisc::CPUHasHwAesSupport()
 	return DetectCPUFeatures().bHasAes;
 }
 
+FString FIOSPlatformMisc::GetDiscardableCacheDir()
+{
+	FString CachesDirectory = FString([NSSearchPathForDirectoriesInDomains(NSCachesDirectory, NSUserDomainMask, YES) objectAtIndex:0]) + TEXT("/");
+	return CachesDirectory;
+}
+
 #if !UE_BUILD_SHIPPING
 bool FIOSPlatformMisc::IsConsoleOpen()
 {
diff --git a/Engine/Source/Runtime/Core/Public/IOS/IOSPlatformMisc.h b/Engine/Source/Runtime/Core/Public/IOS/IOSPlatformMisc.h
--- a/Engine/Source/Runtime/Core/Public/IOS/IOSPlatformMisc.h
+++ b/Engine/Source/Runtime/Core/Public/IOS/IOSPlatformMisc.h
@@ -244,6 +244,8 @@ struct CORE_API FIOSPlatformMisc : public FApplePlatformMisc
 	static bool CPUHasHwCrcSupport();
 	static bool CPUHasHwAesSupport();
 
+	static FString GetDiscardableCacheDir();
+
 #if !UE_BUILD_SHIPPING
 	static bool IsConsoleOpen();
 #endif
diff --git a/Engine/Source/Runtime/StorageServerClient/Private/StorageServerConnection.cpp b/Engine/Source/Runtime/StorageServerClient/Private/StorageServerConnection.cpp
--- a/Engine/Source/Runtime/StorageServerClient/Private/StorageServerConnection.cpp
+++ b/Engine/Source/Runtime/StorageServerClient/Private/StorageServerConnection.cpp
@@ -872,6 +872,8 @@ void FStorageServerConnection::SetupCacheStrategy()
 	DevStoragePath = STORAGE_SERVER_PLATFORM_CACHE_ROOT_DIR;
 #elif PLATFORM_ANDROID
 	DevStoragePath = FPlatformMisc::GamePersistentDownloadDir();
+#elif PLATFORM_IOS
+	DevStoragePath = FPlatformMisc::GetDiscardableCacheDir();
 #endif
 
 	const uint64 Size = ((uint64)Configuration.CacheSizeKB) * 1024;
diff --git a/Engine/Source/Runtime/StorageServerClient/Private/StorageServerConnection.h b/Engine/Source/Runtime/StorageServerClient/Private/StorageServerConnection.h
--- a/Engine/Source/Runtime/StorageServerClient/Private/StorageServerConnection.h
+++ b/Engine/Source/Runtime/StorageServerClient/Private/StorageServerConnection.h
@@ -18,7 +18,7 @@
 #include "Cache/CacheStrategy.h"
 
 #ifndef PLATFORM_SUPPORTS_STORAGE_SERVER_CACHE
-	#define PLATFORM_SUPPORTS_STORAGE_SERVER_CACHE (PLATFORM_WINDOWS || PLATFORM_ANDROID)
+	#define PLATFORM_SUPPORTS_STORAGE_SERVER_CACHE (PLATFORM_WINDOWS || PLATFORM_ANDROID || PLATFORM_IOS)
 #endif
 
 #ifndef PLATFORM_ENABLES_STORAGE_SERVER_CACHE_BY_DEFAULT
