Enable zen caching on more platforms

#rnx

diff --git a/Engine/Source/Runtime/StorageServerClient/Private/StorageServerConnection.cpp b/Engine/Source/Runtime/StorageServerClient/Private/StorageServerConnection.cpp
--- a/Engine/Source/Runtime/StorageServerClient/Private/StorageServerConnection.cpp
+++ b/Engine/Source/Runtime/StorageServerClient/Private/StorageServerConnection.cpp
@@ -854,7 +854,7 @@ void FStorageServerConnection::GetDefaultCacheConfiguration(FCacheConfiguration&
 	OutConfiguration.bUseMemoryMappedStorage = FPlatformProperties::SupportsMemoryMappedFiles() && (FCString::Strstr(CmdLine, TEXT("-ZenNoMmappedStorage")) == nullptr);
 }
 
-#if !HAS_STORAGE_SERVER_PLATFORM_CACHE_STRATEGY
+#if !PLATFORM_HAS_CUSTOM_STORAGE_SERVER_CACHE_STRATEGY
 
 void FStorageServerConnection::SetupCacheStrategy()
 {
@@ -907,7 +907,7 @@ void FStorageServerConnection::SetupCacheStrategy()
 
 	AsyncQueryLatestServerChunkInfo = MakeShared<FAsyncQueryLatestServerChunkInfo>(*this);
 }
-#endif // !HAS_STORAGE_SERVER_PLATFORM_CACHE_STRATEGY
+#endif // !PLATFORM_HAS_CUSTOM_STORAGE_SERVER_CACHE_STRATEGY
 
 bool FStorageServerConnection::FinalizeSetupCacheStrategy()
 {
diff --git a/Engine/Source/Runtime/StorageServerClient/Private/StorageServerConnection.h b/Engine/Source/Runtime/StorageServerClient/Private/StorageServerConnection.h
--- a/Engine/Source/Runtime/StorageServerClient/Private/StorageServerConnection.h
+++ b/Engine/Source/Runtime/StorageServerClient/Private/StorageServerConnection.h
@@ -21,12 +21,16 @@
 	#define PLATFORM_SUPPORTS_STORAGE_SERVER_CACHE (PLATFORM_WINDOWS || PLATFORM_ANDROID || PLATFORM_IOS)
 #endif
 
+#if PLATFORM_SUPPORTS_STORAGE_SERVER_CACHE && PLATFORM_ANDROID
+	#define PLATFORM_ENABLES_STORAGE_SERVER_CACHE_BY_DEFAULT 1
+#endif
+
 #ifndef PLATFORM_ENABLES_STORAGE_SERVER_CACHE_BY_DEFAULT
 	#define PLATFORM_ENABLES_STORAGE_SERVER_CACHE_BY_DEFAULT 0
 #endif
 
-#ifndef HAS_STORAGE_SERVER_PLATFORM_CACHE_STRATEGY
-	#define HAS_STORAGE_SERVER_PLATFORM_CACHE_STRATEGY 0
+#ifndef PLATFORM_HAS_CUSTOM_STORAGE_SERVER_CACHE_STRATEGY
+	#define PLATFORM_HAS_CUSTOM_STORAGE_SERVER_CACHE_STRATEGY 0
 #endif
 
 #ifndef HAS_STORAGE_SERVER_RPC_GETCHUNKS_API // zen server 5.5.16 introduced a new API to request chunks
