- Reworked the IterateDirectory functions in FStorageServerPlatformFile to not return the same file if it's in both local and remote locations. Refactored to share code between the 3 versions of the Iterate functions into a new FUniqueDirectoryStatVisitor class

diff --git a/Engine/Source/Runtime/StorageServerClient/Private/StorageServerPlatformFile.cpp b/Engine/Source/Runtime/StorageServerClient/Private/StorageServerPlatformFile.cpp
--- a/Engine/Source/Runtime/StorageServerClient/Private/StorageServerPlatformFile.cpp
+++ b/Engine/Source/Runtime/StorageServerClient/Private/StorageServerPlatformFile.cpp
@@ -1358,106 +1358,182 @@ IFileHandle* FStorageServerPlatformFile::OpenRead(const TCHAR* Filename, bool bA
 	return (!GPreferLocalForNonAssets && LowerLevel && IsNonServerFilenameAllowed(Filename)) ? LowerLevel->OpenRead(Filename, bAllowWrite) : nullptr;
 }
 
-bool FStorageServerPlatformFile::IterateDirectory(const TCHAR* Directory, IPlatformFile::FDirectoryVisitor& Visitor)
+
+// a reusable class that handles iterating two different locations that will return results that appear
+// to the engine as one location (in this case the StorageServer and Local Files will have the same path for UE)
+// and we need to only return 1 copy, otherwise we can cause errors ot duplicated work)
+template<typename ParentVisitorClass, typename DataType>
+class FUniqueDirectoryStatVisitor : public ParentVisitorClass
 {
-	TStringBuilder<1024> StorageServerDirectory;
-	bool bShouldContinue = true;
-	bool bIterateOnServer = MakeStorageServerPath(Directory, StorageServerDirectory) && ServerToc.DirectoryExists(*StorageServerDirectory);
+private:
+	TSet<FString> AlreadyVisited;
+	ParentVisitorClass& RealVisitor;
+	IPlatformFile* LowerLevel;
+	FStorageServerPlatformFile& ServerPlatformFile;
 
-	if (bIterateOnServer)
+public:
+	FUniqueDirectoryStatVisitor(FStorageServerPlatformFile* InPlatformFile, ParentVisitorClass& InVisitor)
+		: RealVisitor(InVisitor)
+		, LowerLevel(InPlatformFile->GetLowerLevel())
+		, ServerPlatformFile(*InPlatformFile)
 	{
-		// a false return here means the visitor returned false, which means to stop iterating on more files
-		// it can also mean the directory doesn't exist, but we just checked that above, so false only means to stop iterating
-		bShouldContinue = ServerToc.IterateDirectory(*StorageServerDirectory, [this, &Visitor](const FIoChunkId& FileChunkId, const TCHAR* FilenameOrDirectory, int64 RawSize)
-			{
-				TStringBuilder<1024> LocalPath;
-				bool bConverted = MakeLocalPath(FilenameOrDirectory, LocalPath);
-				check(bConverted);
-				const bool bDirectory = !FileChunkId.IsValid();
-				return Visitor.CallShouldVisitAndVisit(*LocalPath, bDirectory);
-			});
 	}
 
-	// run a local iterate in 2 cases:
-	//   1. we prefer local and we were not told to stop iterating above (matches FSandboxPlatformFile)
-	//   2. we don't prefer local, and the directory doesn't exist on the server (this is the existing behavior, before we added GPreferLocalForNonAssets)
-	// this code assumes (as does Sandbox seemingly) that the visitors can handle the same path being returned twice
-	if ((GPreferLocalForNonAssets && bShouldContinue) || (!GPreferLocalForNonAssets && !bIterateOnServer))
+	virtual bool ShouldVisitLeafPathname(FStringView LeafPathname)
 	{
-		bShouldContinue = LowerLevel->IterateDirectory(Directory, Visitor);
+		return RealVisitor.ShouldVisitLeafPathname(LeafPathname);
 	}
-	return bShouldContinue;
-}
 
-bool FStorageServerPlatformFile::IterateDirectoryRecursively(const TCHAR* Directory, IPlatformFile::FDirectoryVisitor& Visitor)
-{
-	TStringBuilder<1024> StorageServerDirectory;
-	bool bShouldContinue = true;
-	bool bIterateOnServer = MakeStorageServerPath(Directory, StorageServerDirectory) && ServerToc.DirectoryExists(*StorageServerDirectory);
-
-	// see IterateDirectory() comments for the flow logic
-	if (bIterateOnServer)
+	bool Visit(const TCHAR* FilenameOrDirectory, DataType StatData)
 	{
-		bShouldContinue = ServerToc.IterateDirectoryRecursively(*StorageServerDirectory, [this, &Visitor](const FIoChunkId& FileChunkId, const TCHAR* FilenameOrDirectory, int64 RawSize)
+		// for speed reasons, we only do the double location checks when we enable the new
+		// PreferLocal behavior - if we ever have duped results without it, we can remove 
+		// this check
+		if (GPreferLocalForNonAssets)
 		{
-			TStringBuilder<1024> LocalPath;
-			bool bConverted = MakeLocalPath(FilenameOrDirectory, LocalPath);
-			check(bConverted);
-			const bool bDirectory = !FileChunkId.IsValid();
-			return Visitor.CallShouldVisitAndVisit(*LocalPath, bDirectory);
-		});
-	}
-	if ((GPreferLocalForNonAssets && bShouldContinue) || (!GPreferLocalForNonAssets && !bIterateOnServer))
-	{
-		bShouldContinue = LowerLevel->IterateDirectoryRecursively(Directory, Visitor);
+			FString FileStr = FPaths::ConvertRelativePathToFull(FilenameOrDirectory);
+
+			if (AlreadyVisited.Contains(FileStr))
+			{
+				return true;
+			}
+			AlreadyVisited.Add(FileStr);
+		}
+		return RealVisitor.Visit(FilenameOrDirectory, StatData);
 	}
-	return bShouldContinue;
-}
 
-bool FStorageServerPlatformFile::IterateDirectoryStat(const TCHAR* Directory, FDirectoryStatVisitor& Visitor)
-{
-	TStringBuilder<1024> StorageServerDirectory;
-	bool bShouldContinue = true;
-	bool bIterateOnServer = MakeStorageServerPath(Directory, StorageServerDirectory) && ServerToc.DirectoryExists(*StorageServerDirectory);
 
-	// see IterateDirectory() comments for the flow logic
-	if (bIterateOnServer)
+	bool PerformMergedIteration(const TCHAR* Directory,
+		TFunctionRef<bool(const TCHAR* Directory, ParentVisitorClass& Visitor)> LowLevelOperation,
+		TFunctionRef<bool(const TCHAR* Directory, ParentVisitorClass& Visitor)> RemoteOperation)
 	{
-		bShouldContinue = ServerToc.IterateDirectory(*StorageServerDirectory, [this, &Visitor](const FIoChunkId& FileChunkId, const TCHAR* ServerFilenameOrDirectory, int64 RawSize)
+		// first, if we prefer local assets, iterate on local first (in case the data doesn't match remote for whatever reason,
+		// local data will be used)
+		if (GPreferLocalForNonAssets)
 		{
-			TStringBuilder<1024> LocalPath;
-			bool bConverted = MakeLocalPath(ServerFilenameOrDirectory, LocalPath);
-			check(bConverted);
-			FFileStatData FileStatData;
-			if (FileChunkId.IsValid())
+			if (LowLevelOperation(Directory, *this) == false)
 			{
-				FileStatData = FFileStatData(
-					FDateTime::Now(),
-					FDateTime::Now(),
-					FDateTime::Now(),
-					RawSize,
-					false,
-					true);
-				check(FileStatData.bIsValid);
+				return false;
 			}
-			else
+		}
+
+		// then look remote
+		TStringBuilder<1024> StorageServerDirectory;
+		bool bShouldContinue = true;
+		bool bIterateOnServer = ServerPlatformFile.MakeStorageServerPath(Directory, StorageServerDirectory) && ServerPlatformFile.ServerToc.DirectoryExists(*StorageServerDirectory);
+		if (bIterateOnServer)
+		{
+			if (RemoteOperation(*StorageServerDirectory, *this) == false)
 			{
-				FileStatData = FFileStatData(
-					FDateTime::MinValue(),
-					FDateTime::MinValue(),
-					FDateTime::MinValue(),
-					0,
-					true,
-					true);
+				return false;
 			}
-			return Visitor.CallShouldVisitAndVisit(*LocalPath, FileStatData);
-		});
-	}
-	if ((GPreferLocalForNonAssets && bShouldContinue) || (!GPreferLocalForNonAssets && !bIterateOnServer))
-	{
-		bShouldContinue = LowerLevel->IterateDirectoryStat(Directory, Visitor);
+		}
+
+		// finall look locally if we are preferring remote over local
+		if (!GPreferLocalForNonAssets)
+		{
+			if (LowLevelOperation(Directory, *this) == false)
+			{
+				return false;
+			}
+		}
+
+		return true;
 	}
-	return bShouldContinue;
+
+};
+
+
+
+bool FStorageServerPlatformFile::IterateDirectory(const TCHAR* Directory, IPlatformFile::FDirectoryVisitor& Visitor)
+{
+	FUniqueDirectoryStatVisitor<FDirectoryVisitor, bool> MergedVisitor(this, Visitor);
+
+	return MergedVisitor.PerformMergedIteration(Directory,
+		[this](const TCHAR* Directory, FDirectoryVisitor& Visitor)
+		{
+			return LowerLevel->IterateDirectory(Directory, Visitor);
+		},
+		[this](const TCHAR* Directory, FDirectoryVisitor& Visitor)
+		{
+			return ServerToc.IterateDirectory(Directory, [this, &Visitor](const FIoChunkId& FileChunkId, const TCHAR* FilenameOrDirectory, int64 RawSize)
+				{
+					TStringBuilder<1024> LocalPath;
+					bool bConverted = MakeLocalPath(FilenameOrDirectory, LocalPath);
+					check(bConverted);
+					const bool bDirectory = !FileChunkId.IsValid();
+					return Visitor.CallShouldVisitAndVisit(*LocalPath, bDirectory);
+				});
+		}
+	);
+}
+
+bool FStorageServerPlatformFile::IterateDirectoryRecursively(const TCHAR* Directory, IPlatformFile::FDirectoryVisitor& Visitor)
+{
+	FUniqueDirectoryStatVisitor<FDirectoryVisitor, bool> MergedVisitor(this, Visitor);
+
+	return MergedVisitor.PerformMergedIteration(Directory,
+		[this](const TCHAR* Directory, FDirectoryVisitor& Visitor)
+		{
+			return LowerLevel->IterateDirectoryRecursively(Directory, Visitor);
+		},
+		[this](const TCHAR* Directory, FDirectoryVisitor& Visitor)
+		{
+			return ServerToc.IterateDirectoryRecursively(Directory, [this, &Visitor](const FIoChunkId& FileChunkId, const TCHAR* FilenameOrDirectory, int64 RawSize)
+				{
+					TStringBuilder<1024> LocalPath;
+					bool bConverted = MakeLocalPath(FilenameOrDirectory, LocalPath);
+					check(bConverted);
+					const bool bDirectory = !FileChunkId.IsValid();
+					return Visitor.CallShouldVisitAndVisit(*LocalPath, bDirectory);
+				});
+		}
+	);
+}
+
+
+bool FStorageServerPlatformFile::IterateDirectoryStat(const TCHAR* Directory, FDirectoryStatVisitor& Visitor)
+{
+	FUniqueDirectoryStatVisitor<FDirectoryStatVisitor, const FFileStatData&> MergedVisitor(this, Visitor);
+
+	return MergedVisitor.PerformMergedIteration(Directory,
+		[this](const TCHAR* Directory, FDirectoryStatVisitor& Visitor)
+		{
+			return LowerLevel->IterateDirectoryStat(Directory, Visitor);
+		},
+		[this](const TCHAR* Directory, FDirectoryStatVisitor& Visitor)
+		{
+			return ServerToc.IterateDirectory(Directory, [this, &Visitor](const FIoChunkId& FileChunkId, const TCHAR* ServerFilenameOrDirectory, int64 RawSize)
+				{
+					TStringBuilder<1024> LocalPath;
+					bool bConverted = MakeLocalPath(ServerFilenameOrDirectory, LocalPath);
+					check(bConverted);
+					FFileStatData FileStatData;
+					if (FileChunkId.IsValid())
+					{
+						FileStatData = FFileStatData(
+							FDateTime::Now(),
+							FDateTime::Now(),
+							FDateTime::Now(),
+							RawSize,
+							false,
+							true);
+						check(FileStatData.bIsValid);
+					}
+					else
+					{
+						FileStatData = FFileStatData(
+							FDateTime::MinValue(),
+							FDateTime::MinValue(),
+							FDateTime::MinValue(),
+							0,
+							true,
+							true);
+					}
+					return Visitor.CallShouldVisitAndVisit(*LocalPath, FileStatData); 
+				});
+		}
+	);
 }
 
 FOpenMappedResult FStorageServerPlatformFile::OpenMappedEx(const TCHAR* Filename, EOpenReadFlags OpenOptions, int64 MaximumSize)
diff --git a/Engine/Source/Runtime/StorageServerClient/Private/StorageServerPlatformFile.h b/Engine/Source/Runtime/StorageServerClient/Private/StorageServerPlatformFile.h
--- a/Engine/Source/Runtime/StorageServerClient/Private/StorageServerPlatformFile.h
+++ b/Engine/Source/Runtime/StorageServerClient/Private/StorageServerPlatformFile.h
@@ -127,6 +127,8 @@ public:
 	}
 private:
 	friend class FStorageServerFileHandle;
+	template<typename ParentVisitorClass, typename DataType> friend class FUniqueDirectoryStatVisitor;
+
 
 	void InitializeConnection();
 	bool IsNonServerFilenameAllowed(FStringView InFilename);
