Zen streaming: Use non-batched requests if connected to a workspace endpoint.

#rnx

diff --git a/Engine/Source/Runtime/StorageServerClient/Private/StorageServerConnection.cpp b/Engine/Source/Runtime/StorageServerClient/Private/StorageServerConnection.cpp
--- a/Engine/Source/Runtime/StorageServerClient/Private/StorageServerConnection.cpp
+++ b/Engine/Source/Runtime/StorageServerClient/Private/StorageServerConnection.cpp
@@ -92,6 +92,10 @@ bool FStorageServerConnection::Initialize(TArrayView<const FString> HostAddresse
 {
 	BaseURI = InBaseURI;
 	TArray<FString> SortedHostAddresses = SortHostAddressesByLocalSubnet(HostAddresses, Port);
+	if (BaseURI.StartsWith("/ws/"))
+	{
+		bIsUsingZenWorkspace = true;
+	}
 
 	for (const FString& HostAddress : SortedHostAddresses)
 	{
@@ -382,20 +386,33 @@ TIoStatusOr<FIoBuffer> FStorageServerConnection::ReadChunkRequest(
 	else
 	{
 #if HAS_STORAGE_SERVER_RPC_GETCHUNKS_API
-		TArray<FChunkBatchRequestEntry> ChunkBatchRequests;
-		ChunkBatchRequests.Add(FChunkBatchRequestEntry::DataRequest(ChunkId, Offset, Size));
-
-		FIoStatus ResultStatus = ReadChunkBatchRequest(ChunkBatchRequests, [&](FIoChunkId Id, EStorageServerContentType MimeType, FIoBuffer Data, const TOptional<uint64>& ModTag)
+		if (bIsUsingZenWorkspace)
 		{
-			ensure(ChunkId == Id);
-			Data.MakeOwned();
-			ResultTuple = IStorageServerHttpClient::FResult(Data, MimeType);
-			ResultModTag = ModTag.GetValue(); // ModTag must be present for cache invalidation to work
-		});
+			TRACE_CPUPROFILER_EVENT_SCOPE(FStorageServerConnection::ReadChunkRequest::Http);
 
-		if (!ResultStatus.IsOk())
+			TAnsiStringBuilder<256> ResourceBuilder;
+			BuildReadChunkRequestUrl(ResourceBuilder, ChunkId, Offset, Size);
+
+			ResultTuple = HttpClient->RequestSync(*ResourceBuilder);
+			ResultModTag = 0; // This endpoint doesn't support ModTag.
+		}
+		else
 		{
-			ResultTuple = IStorageServerHttpClient::FResult(ResultStatus, EStorageServerContentType::Unknown);
+			TArray<FChunkBatchRequestEntry> ChunkBatchRequests;
+			ChunkBatchRequests.Add(FChunkBatchRequestEntry::DataRequest(ChunkId, Offset, Size));
+
+			FIoStatus ResultStatus = ReadChunkBatchRequest(ChunkBatchRequests, [&](FIoChunkId Id, EStorageServerContentType MimeType, FIoBuffer Data, const TOptional<uint64>& ModTag)
+				{
+					ensure(ChunkId == Id);
+					Data.MakeOwned();
+					ResultTuple = IStorageServerHttpClient::FResult(Data, MimeType);
+					ResultModTag = ModTag.GetValue(); // ModTag must be present for cache invalidation to work
+				});
+
+			if (!ResultStatus.IsOk())
+			{
+				ResultTuple = IStorageServerHttpClient::FResult(ResultStatus, EStorageServerContentType::Unknown);
+			}
 		}
 #else
 		TRACE_CPUPROFILER_EVENT_SCOPE(FStorageServerConnection::ReadChunkRequest::Http);
diff --git a/Engine/Source/Runtime/StorageServerClient/Private/StorageServerConnection.h b/Engine/Source/Runtime/StorageServerClient/Private/StorageServerConnection.h
--- a/Engine/Source/Runtime/StorageServerClient/Private/StorageServerConnection.h
+++ b/Engine/Source/Runtime/StorageServerClient/Private/StorageServerConnection.h
@@ -34,7 +34,7 @@
 #endif
 
 #ifndef HAS_STORAGE_SERVER_RPC_GETCHUNKS_API // zen server 5.5.16 introduced a new API to request chunks
-	#define HAS_STORAGE_SERVER_RPC_GETCHUNKS_API 0 // #jira UE-234828 temporarily disabled until RPC batch requests support gets added to zen workspaces
+	#define HAS_STORAGE_SERVER_RPC_GETCHUNKS_API 1
 #endif
 
 #if !UE_BUILD_SHIPPING
@@ -134,6 +134,7 @@ private:
 	TUniquePtr<StorageServer::ICacheStrategy> CacheStrategy;
 	FAnsiString BaseURI;
 	FString CurrentHostAddr;
+	bool bIsUsingZenWorkspace = false;		    // is the connection to the /ws/ endpoint
 
 	// Stats
 	std::atomic<uint64> AccumulatedBytes = 0;
