Zen storage server startup precache:
- use non-batched requests if connected to a Zen workspace.
- fix improper cache size calculation in case of a request failure

#rnx

diff --git a/Engine/Source/Runtime/StorageServerClient/Private/StorageServerConnection.h b/Engine/Source/Runtime/StorageServerClient/Private/StorageServerConnection.h
--- a/Engine/Source/Runtime/StorageServerClient/Private/StorageServerConnection.h
+++ b/Engine/Source/Runtime/StorageServerClient/Private/StorageServerConnection.h
@@ -73,6 +73,10 @@ public:
 
 	TIoStatusOr<Workspaces> GetWorkspaces();
 	TIoStatusOr<FString> CreateShare(const FString& WorkspaceId, const FString& SharePath, const FString& Alias);
+	bool IsConnectedToWorkspace() const
+	{
+		return bIsUsingZenWorkspace;
+	}
 
 	void PackageStoreRequest(TFunctionRef<void(FPackageStoreEntryResource&&)> Callback);
 	void FileManifestRequest(TFunctionRef<void(FIoChunkId Id, FStringView Path, int64 RawSize)> Callback);
diff --git a/Engine/Source/Runtime/StorageServerClient/Private/StorageServerPlatformFile.cpp b/Engine/Source/Runtime/StorageServerClient/Private/StorageServerPlatformFile.cpp
--- a/Engine/Source/Runtime/StorageServerClient/Private/StorageServerPlatformFile.cpp
+++ b/Engine/Source/Runtime/StorageServerClient/Private/StorageServerPlatformFile.cpp
@@ -256,11 +256,17 @@ public:
 		, IsCompleted(FPlatformProcess::GetSynchEventFromPool())
 		, bExitWhenComplete(false)
 		, bExitImmediately(false)
+		, bUseBatchedRequests(true)
 	{
 		if (FRunnableThread::Create( this, TEXT("StorageServerPrecache"), 0, EThreadPriority::TPri_Normal) == nullptr)
 		{
 			IsCompleted->Trigger();
 		}
+
+		if (Connection.IsConnectedToWorkspace())
+		{
+			bUseBatchedRequests = false;
+		}
 	}
 
 	virtual ~FStorageServerEngineStartupPrecache()
@@ -374,9 +380,12 @@ private:
 			DecompressedData.Append(Data.GetData(), Data.GetSize());
 		}
 
-		FWriteScopeLock Lock(CacheLock);
-		Cache.Add(ChunkId, DecompressedData);
-
+		if (CacheSize + DecompressedData.Num() <= MaxCacheSize)
+		{
+			FWriteScopeLock Lock(CacheLock);
+			CacheSize += DecompressedData.Num();
+			Cache.Add(ChunkId, DecompressedData);
+		}
 	}
 
 #if HAS_STORAGE_SERVER_RPC_GETCHUNKS_API
@@ -400,8 +409,25 @@ private:
 				break;
 			}
 
-			CacheSize += Size;
-			Items.Add(FStorageServerConnection::FChunkBatchRequestEntry{Id, 0, uint64(Size)});
+			if (bUseBatchedRequests)
+			{
+				Items.Add(FStorageServerConnection::FChunkBatchRequestEntry{Id, 0, uint64(Size)});
+			}
+			else
+			{
+				TArray<uint8> TempBuffer;
+				TempBuffer.AddUninitialized(Size);
+				TIoStatusOr<FIoBuffer> Result = Connection.ReadChunkRequest(Id, 0, Size, FIoBuffer(FIoBuffer::Wrap, TempBuffer.GetData(), Size), false);
+				if (Result.IsOk())
+				{
+					TempBuffer.SetNum(Result.ValueOrDie().GetSize());
+					CacheSize += TempBuffer.Num();
+					{
+						FWriteScopeLock Lock(CacheLock);
+						Cache.Add(Id, MoveTemp(TempBuffer));
+					}
+				}
+			}
 		}
 
 		if (Items.IsEmpty())
@@ -434,6 +460,7 @@ private:
 
 	std::atomic<bool> bExitWhenComplete;
 	std::atomic<bool> bExitImmediately;
+	std::atomic<bool> bUseBatchedRequests;
 	
 #if STORAGE_SERVER_START_CACHE_REPORT
 	TSet<FIoChunkId> UncachedChunks;
