Use hostname as the last resort for client to connect to zen server

diff --git a/Engine/Source/Developer/IoStoreUtilities/Private/ZenStoreWriter.cpp b/Engine/Source/Developer/IoStoreUtilities/Private/ZenStoreWriter.cpp
--- a/Engine/Source/Developer/IoStoreUtilities/Private/ZenStoreWriter.cpp
+++ b/Engine/Source/Developer/IoStoreUtilities/Private/ZenStoreWriter.cpp
@@ -471,6 +471,18 @@ void FZenStoreWriter::Initialize(const FCookInfo& Info)
 						{
 							Writer->WriteValue(Address->ToString(false));
 						}
+						FString MachineHostName;
+						if (SocketSubsystem->GetHostName(MachineHostName))
+						{
+							// Try to acquire FQDN
+							FAddressInfoResult GAIRequest = SocketSubsystem->GetAddressInfo(*MachineHostName, nullptr, EAddressInfoFlags::AllResultsWithMapping | EAddressInfoFlags::OnlyUsableAddresses | EAddressInfoFlags::AllResults | EAddressInfoFlags::FQDomainName, NAME_None);
+							if (GAIRequest.ReturnCode == SE_NO_ERROR && !GAIRequest.CanonicalNameResult.IsEmpty())
+							{
+								MachineHostName = GAIRequest.CanonicalNameResult;
+							}
+
+							Writer->WriteValue(TEXT("hostname://") + MachineHostName);
+						}
 #if PLATFORM_MAC
 						// Store the Bonjour hostname when on a mac so that the client can attempt to connect via USB.
 						// This is done since the "Link Local" IP of the USB connection changes everytime
diff --git a/Engine/Source/Programs/AutomationTool/BuildGraph/Tasks/ZenImportSnapshotTask.cs b/Engine/Source/Programs/AutomationTool/BuildGraph/Tasks/ZenImportSnapshotTask.cs
--- a/Engine/Source/Programs/AutomationTool/BuildGraph/Tasks/ZenImportSnapshotTask.cs
+++ b/Engine/Source/Programs/AutomationTool/BuildGraph/Tasks/ZenImportSnapshotTask.cs
@@ -381,9 +381,35 @@ namespace AutomationTool.Tasks
 			}
 		}
 
+		private static string GetHostNameFQDN()
+		{
+			string hostName = Dns.GetHostName();
+
+			// take first hostname that contains '.'
+			if (!hostName.Contains('.', StringComparison.InvariantCulture))
+			{
+				IPAddress[] ipAddresses = Dns.GetHostAddresses("");
+
+				foreach (IPAddress ipAddress in ipAddresses)
+				{
+					string hostNameTemp = Dns.GetHostEntry(ipAddress).HostName;
+					if (hostNameTemp.Contains('.', StringComparison.InvariantCulture))
+					{
+						hostName = hostNameTemp;
+						break;
+					}
+				}
+			}
+
+			return hostName;
+		}
+
 		private static string[] GetHostAddresses()
 		{
-			return Dns.GetHostAddresses("", System.Net.Sockets.AddressFamily.InterNetwork).Select(addr => addr.ToString()).ToArray();
+			return Dns.GetHostAddresses("", System.Net.Sockets.AddressFamily.InterNetwork)
+				.Select(addr => addr.ToString())
+				.Append($"hostname://{GetHostNameFQDN()}")
+				.ToArray();
 		}
 
 		/// <summary>
diff --git a/Engine/Source/Runtime/StorageServerClient/Private/StorageServerConnection.cpp b/Engine/Source/Runtime/StorageServerClient/Private/StorageServerConnection.cpp
--- a/Engine/Source/Runtime/StorageServerClient/Private/StorageServerConnection.cpp
+++ b/Engine/Source/Runtime/StorageServerClient/Private/StorageServerConnection.cpp
@@ -730,16 +730,23 @@ TArray<FString> FStorageServerConnection::SortHostAddressesByLocalSubnet(TArrayV
 		localAddrStringSubnet = localAddrStringSubnet.LeftChop(localAddrStringSubnet.Len() - localLastDotPos);
 	}
 
-	TArray<FString> PlatformSocketAddress;
+	TArray<FString> PlatformSocketAddresses;
 	TArray<FString> IPV6Loopback;
 	TArray<FString> IPV4Loopback;
 	TArray<FString> RegularAddresses;
+	TArray<FString> HostnameAddresses;
 
 	for (const FString& HostAddress : HostAddresses)
 	{
 		if (IsPlatformSocketAddress(HostAddress))
 		{
-			PlatformSocketAddress.Push(HostAddress);
+			PlatformSocketAddresses.Push(HostAddress);
+			continue;
+		}
+
+		if (IsHostnameAddress(HostAddress))
+		{
+			HostnameAddresses.Push(HostAddress);
 			continue;
 		}
 
@@ -800,25 +807,11 @@ TArray<FString> FStorageServerConnection::SortHostAddressesByLocalSubnet(TArrayV
 		}
 	}
 
-	for (const FString& PlatformAddr : PlatformSocketAddress)
-	{
-		Result.Push(PlatformAddr);
-	}
-
-	for (const FString& Addrv6lb : IPV6Loopback)
-	{
-		Result.Push(Addrv6lb);
-	}
-
-	for (const FString& Addrv4lb : IPV4Loopback)
-	{
-		Result.Push(Addrv4lb);
-	}
-
-	for (const FString& RegularAddr : RegularAddresses)
-	{
-		Result.Push(RegularAddr);
-	}
+	Result.Append(PlatformSocketAddresses);
+	Result.Append(IPV6Loopback);
+	Result.Append(IPV4Loopback);
+	Result.Append(RegularAddresses);
+	Result.Append(HostnameAddresses);
 
 	return Result;
 }
@@ -828,6 +821,11 @@ bool FStorageServerConnection::IsPlatformSocketAddress(const FString Address)
 	return Address.StartsWith(TEXT("platform://"));
 }
 
+bool FStorageServerConnection::IsHostnameAddress(const FString Address)
+{
+	return Address.StartsWith(TEXT("hostname://"));
+}
+
 TUniquePtr<IStorageServerHttpClient> FStorageServerConnection::CreateHttpClient(const FString Address, const int32 Port)
 {
 	TSharedPtr<FInternetAddr> Addr = StringToInternetAddr(Address, Port);
@@ -854,17 +852,24 @@ TSharedPtr<FInternetAddr> FStorageServerConnection::StringToInternetAddr(const F
 {
 	TSharedPtr<FInternetAddr> Result = TSharedPtr<FInternetAddr>();
 
+	ISocketSubsystem& SocketSubsystem = *ISocketSubsystem::Get();
+
 	if (IsPlatformSocketAddress(HostAddr))
 	{
 		return Result;
 	}
 
-	ISocketSubsystem& SocketSubsystem = *ISocketSubsystem::Get();
-
 	// Numeric IPV6 addresses can be enclosed in brackets, and must have the brackets stripped before calling GetAddressFromString
 	FString ModifiedHostAddr;
 	const FString* EffectiveHostAddr = &HostAddr;
-	if (!HostAddr.IsEmpty() && HostAddr[0] == TEXT('[') && HostAddr[HostAddr.Len() - 1] == TEXT(']'))
+	bool bIsHostname = false;
+	if (IsHostnameAddress(HostAddr))
+	{
+		ModifiedHostAddr = FStringView(HostAddr).RightChop(11); // cut off "hostname://"
+		EffectiveHostAddr = &ModifiedHostAddr;
+		bIsHostname = true;
+	}
+	else if (!HostAddr.IsEmpty() && HostAddr[0] == TEXT('[') && HostAddr[HostAddr.Len() - 1] == TEXT(']'))
 	{
 #if PLATFORM_HAS_BSD_SOCKETS && !PLATFORM_HAS_BSD_IPV6_SOCKETS
 		// If the platform doesn't have IPV6 BSD Sockets, then handle an attempt at conversion of loopback addresses, and skip and warn about other addresses
@@ -884,7 +889,10 @@ TSharedPtr<FInternetAddr> FStorageServerConnection::StringToInternetAddr(const F
 		EffectiveHostAddr = &ModifiedHostAddr;
 	}
 
-	Result = SocketSubsystem.GetAddressFromString(*EffectiveHostAddr);
+	if (!bIsHostname)
+	{
+		Result = SocketSubsystem.GetAddressFromString(*EffectiveHostAddr);
+	}
 	if (!Result.IsValid() || !Result->IsValid())
 	{
 		FAddressInfoResult GAIRequest = SocketSubsystem.GetAddressInfo(**EffectiveHostAddr, nullptr, EAddressInfoFlags::Default, NAME_None);
diff --git a/Engine/Source/Runtime/StorageServerClient/Private/StorageServerConnection.h b/Engine/Source/Runtime/StorageServerClient/Private/StorageServerConnection.h
--- a/Engine/Source/Runtime/StorageServerClient/Private/StorageServerConnection.h
+++ b/Engine/Source/Runtime/StorageServerClient/Private/StorageServerConnection.h
@@ -148,6 +148,7 @@ private:
 
 	TArray<FString> SortHostAddressesByLocalSubnet(TArrayView<const FString> HostAddresses, const int32 Port);
 	static bool IsPlatformSocketAddress(const FString Address); 
+	static bool IsHostnameAddress(const FString Address); 
 	TUniquePtr<IStorageServerHttpClient> CreateHttpClient(const FString Address, const int32 Port);
 	TSharedPtr<FInternetAddr> StringToInternetAddr(const FString Address, const int32 Port);
 	bool HandshakeRequest();
diff --git a/Engine/Source/Runtime/StorageServerClient/Private/StorageServerPlatformFile.cpp b/Engine/Source/Runtime/StorageServerClient/Private/StorageServerPlatformFile.cpp
--- a/Engine/Source/Runtime/StorageServerClient/Private/StorageServerPlatformFile.cpp
+++ b/Engine/Source/Runtime/StorageServerClient/Private/StorageServerPlatformFile.cpp
@@ -1087,6 +1087,7 @@ bool FStorageServerPlatformFile::ShouldBeUsed(IPlatformFile* Inner, const TCHAR*
 		return false;
 
 	bool bPreferFileSystem = false;
+	TArray<FString> HostNames;
 
 #if WITH_COTF
 	UE::Cook::ICookOnTheFlyModule& CookOnTheFlyModule = FModuleManager::LoadModuleChecked<UE::Cook::ICookOnTheFlyModule>(TEXT("CookOnTheFly"));
@@ -1146,6 +1147,11 @@ bool FStorageServerPlatformFile::ShouldBeUsed(IPlatformFile* Inner, const TCHAR*
 							}
 							else
 #endif 
+							if (RemoteHostNameStr.StartsWith("hostname://"))
+							{
+								HostNames.Add(RemoteHostNameStr);
+							}
+							else
 							{
 								HostAddrs.Add(RemoteHostNameStr);
 							}
@@ -1180,6 +1186,8 @@ bool FStorageServerPlatformFile::ShouldBeUsed(IPlatformFile* Inner, const TCHAR*
 	{
 		UE_LOG(LogStorageServerPlatformFile, Display, TEXT("Using connection settings from command line: -ZenStorePort='%d'"), HostPort);
 	}
+	// add hostnames as last resort
+	HostAddrs.Append(HostNames);
 
 	if (!bPreferFileSystem || UE::IsUsingZenPakFileStreaming())
 	{
