Validate zen chunk infos in cache journal on load

#rnx

diff --git a/Engine/Source/Runtime/StorageServerClient/Private/Cache/CacheJournalSectioned.cpp b/Engine/Source/Runtime/StorageServerClient/Private/Cache/CacheJournalSectioned.cpp
--- a/Engine/Source/Runtime/StorageServerClient/Private/Cache/CacheJournalSectioned.cpp
+++ b/Engine/Source/Runtime/StorageServerClient/Private/Cache/CacheJournalSectioned.cpp
@@ -30,7 +30,16 @@ namespace StorageServer
 	{
 		JournalFileHandle.Reset(IPlatformFile::GetPlatformPhysical().OpenWrite(*JournalFilePath, true, true));
 
-		LoadCacheJournal();
+		// If loading the journal failed, invalidate all of the data
+		if (!LoadCacheJournal())
+		{
+			InvalidateAll();
+		}
+		else
+		{
+			// If the journal loaded, there still might be entries without corresponding cache chunk infos. Validate them
+			ValidateJournal();
+		}
 	}
 
 	void FCacheJournalSectioned::Flush(bool bImmediate)
@@ -378,6 +387,31 @@ namespace StorageServer
 		return true;
 	}
 
+	void FCacheJournalSectioned::ValidateJournal()
+	{
+		TRACE_CPUPROFILER_EVENT_SCOPE(FCacheJournalSectioned::ValidateJournal);
+
+		UE::TScopeLock _(CriticalSection);
+
+		TArray<FIoChunkId> InvalidChunkIds;
+		auto ValidateEntryInfo = [&InvalidChunkIds, this](const FIoChunkId& ChunkId, const FCacheEntry& CacheEntry)
+			{
+				FCacheChunkInfo ChunkInfo;
+				if (!TryGetChunkInfo(ChunkId, ChunkInfo))
+				{
+					InvalidChunkIds.Add(ChunkId);
+				}
+			};
+
+		IterateCacheEntries(ValidateEntryInfo);
+		for (const FIoChunkId& ChunkId : InvalidChunkIds)
+		{
+			Invalidate(ChunkId);
+		}
+
+		UE_LOG(LogCacheJournal, Warning, TEXT("There were %d invalid zen cache journal entries. All have been invalidated."), InvalidChunkIds.Num());
+	}
+
 	FJournalChunkInfoPage* FCacheJournalSectioned::FindOrAddChunkInfoPage(const FIoChunkId& ChunkId)
 	{
 		for (FJournalChunkInfoPage& Page : ChunkInfoPages)
diff --git a/Engine/Source/Runtime/StorageServerClient/Private/Cache/CacheJournalSectioned.h b/Engine/Source/Runtime/StorageServerClient/Private/Cache/CacheJournalSectioned.h
--- a/Engine/Source/Runtime/StorageServerClient/Private/Cache/CacheJournalSectioned.h
+++ b/Engine/Source/Runtime/StorageServerClient/Private/Cache/CacheJournalSectioned.h
@@ -231,6 +231,7 @@ namespace StorageServer
 
 	private:
 		bool LoadCacheJournal();
+		void ValidateJournal();
 
 		FJournalChunkInfoPage* FindOrAddChunkInfoPage(const FIoChunkId& ChunkId);
 		FJournalChunkPage* FindOrAddChunkPage(const FIoChunkId& ChunkId);
