Adding a way to configure zen streaming cache size via defines
#rnx

diff --git a/Engine/Source/Runtime/StorageServerClient/Private/StorageServerConnection.cpp b/Engine/Source/Runtime/StorageServerClient/Private/StorageServerConnection.cpp
--- a/Engine/Source/Runtime/StorageServerClient/Private/StorageServerConnection.cpp
+++ b/Engine/Source/Runtime/StorageServerClient/Private/StorageServerConnection.cpp
@@ -33,6 +33,14 @@
 
 #if !UE_BUILD_SHIPPING
 
+#ifndef STORAGE_SERVER_PLATFORM_CACHE_SIZE_KB
+	#define STORAGE_SERVER_PLATFORM_CACHE_SIZE_KB 4 * 1024 * 1024 // 4 GB by default
+#endif
+
+#ifndef STORAGE_SERVER_PLATFORM_CACHE_ABANDON_SIZE_KB
+	#define STORAGE_SERVER_PLATFORM_CACHE_ABANDON_SIZE_KB 512 * 1024 // 512 MB by default
+#endif
+
 DEFINE_LOG_CATEGORY(LogStorageServerConnection);
 
 TRACE_DECLARE_INT_COUNTER(ZenHttpClientSerializedBytes, TEXT("ZenClient/SerializedBytes (compressed)"));
@@ -967,11 +975,11 @@ void FStorageServerConnection::GetDefaultCacheConfiguration(FCacheConfiguration&
 	OutConfiguration.bForceInvalidate = FCString::Strstr(CmdLine, TEXT("-ZenInvalidateCache")) != nullptr;
 	if (!FParse::Value(CmdLine, TEXT("-ZenCacheSizeKB="), OutConfiguration.CacheSizeKB))
 	{
-		OutConfiguration.CacheSizeKB = 4 * 1024 * 1024; // 4 GB by default
+		OutConfiguration.CacheSizeKB = STORAGE_SERVER_PLATFORM_CACHE_SIZE_KB;
 	}
 	if (!FParse::Value(CmdLine, TEXT("-ZenCacheAbandonSizeKB="), OutConfiguration.AbandonSizeKB))
 	{
-		OutConfiguration.AbandonSizeKB = 512 * 1024; // 512 MB by default
+		OutConfiguration.AbandonSizeKB = STORAGE_SERVER_PLATFORM_CACHE_ABANDON_SIZE_KB;
 	}
 	OutConfiguration.FlushInterval = 10.0f;
 	OutConfiguration.FlushEveryNEntries = 0;
