Enforce zenserver refresh file refresh on file manifest request.

This is to account for a change in workspace share's contants after staging to an existing workspace

#rnx

diff --git a/Engine/Source/Runtime/StorageServerClient/Private/StorageServerConnection.cpp b/Engine/Source/Runtime/StorageServerClient/Private/StorageServerConnection.cpp
--- a/Engine/Source/Runtime/StorageServerClient/Private/StorageServerConnection.cpp
+++ b/Engine/Source/Runtime/StorageServerClient/Private/StorageServerConnection.cpp
@@ -257,7 +257,7 @@ void FStorageServerConnection::FileManifestRequest(TFunctionRef<void(FIoChunkId
 	TRACE_CPUPROFILER_EVENT_SCOPE(FStorageServerConnection::FileManifestRequest);
 
 	TAnsiStringBuilder<256> ResourceBuilder;
-	ResourceBuilder.Append(BaseURI).Append("/files?fieldnames=id,clientpath,rawsize");
+	ResourceBuilder.Append(BaseURI).Append("/files?refresh=true&fieldnames=id,clientpath,rawsize");
 
 	IStorageServerHttpClient::FResult ResultTuple = HttpClient->RequestSync(*ResourceBuilder, EStorageServerContentType::CompressedBinary);
 	TIoStatusOr<FIoBuffer> Result = ResultTuple.Get<0>();
