Gracefully handle unresolved hostname:// addresses. When the hostname isn't resolved, this results in a null ServerAddr which in turn crashes the runtime instead of reporting the failure to communicate with the server via a message box.

diff --git a/Engine/Source/Runtime/StorageServerClient/Private/BuiltInHttpClient/BuiltInHttpClientFSocket.cpp b/Engine/Source/Runtime/StorageServerClient/Private/BuiltInHttpClient/BuiltInHttpClientFSocket.cpp
--- a/Engine/Source/Runtime/StorageServerClient/Private/BuiltInHttpClient/BuiltInHttpClientFSocket.cpp
+++ b/Engine/Source/Runtime/StorageServerClient/Private/BuiltInHttpClient/BuiltInHttpClientFSocket.cpp
@@ -99,33 +99,36 @@ IBuiltInHttpClientSocket* FBuiltInHttpClientFSocketPool::AcquireSocket(float Tim
 		return SocketWrapper;
 	}
 
-	FSocket* Socket = SocketSubsystem.CreateSocket(NAME_Stream, TEXT("StorageServer"), ServerAddr->GetProtocolType());
-
-	Socket->SetNoDelay(true);
-
-	if (TimeoutSeconds > 0.0f)
+	if (ServerAddr.IsValid())
 	{
-		Socket->SetNonBlocking(true);
-
-		ON_SCOPE_EXIT
-		{
-			Socket->SetNonBlocking(false);
-		};
-
-		if (Socket->Connect(*ServerAddr) && Socket->Wait(ESocketWaitConditions::WaitForWrite, FTimespan::FromSeconds(TimeoutSeconds)))
+		FSocket* Socket = SocketSubsystem.CreateSocket(NAME_Stream, TEXT("StorageServer"), ServerAddr->GetProtocolType());
+		
+		Socket->SetNoDelay(true);
+		
+		if (TimeoutSeconds > 0.0f)
 		{
-			return new FBuiltInHttpClientFSocket(Socket);
+			Socket->SetNonBlocking(true);
+			
+			ON_SCOPE_EXIT
+			{
+				Socket->SetNonBlocking(false);
+			};
+			
+			if (Socket->Connect(*ServerAddr) && Socket->Wait(ESocketWaitConditions::WaitForWrite, FTimespan::FromSeconds(TimeoutSeconds)))
+			{
+				return new FBuiltInHttpClientFSocket(Socket);
+			}
 		}
-	}
-	else
-	{
-		if (Socket->Connect(*ServerAddr))
+		else
 		{
-			return new FBuiltInHttpClientFSocket(Socket);
+			if (Socket->Connect(*ServerAddr))
+			{
+				return new FBuiltInHttpClientFSocket(Socket);
+			}
 		}
+		
+		delete Socket;
 	}
-
-	delete Socket;
 	return nullptr;
 }
 
@@ -147,4 +150,4 @@ void FBuiltInHttpClientFSocketPool::ReleaseSocket(IBuiltInHttpClientSocket* Sock
 	}
 }
 
-#endif
\ No newline at end of file
+#endif
