Invalidate duplicate and zeroed chunk id entries in Zen streaming cache.

#rnx

diff --git a/Engine/Source/Runtime/StorageServerClient/Private/Cache/CacheJournalSectioned.cpp b/Engine/Source/Runtime/StorageServerClient/Private/Cache/CacheJournalSectioned.cpp
--- a/Engine/Source/Runtime/StorageServerClient/Private/Cache/CacheJournalSectioned.cpp
+++ b/Engine/Source/Runtime/StorageServerClient/Private/Cache/CacheJournalSectioned.cpp
@@ -402,14 +402,50 @@ namespace StorageServer
 					InvalidChunkIds.Add(ChunkId);
 				}
 			};
+		
+		TSet<FIoChunkId> ChunkIdSet;
+		TArray<FIoChunkId> DuplicateChunkIds;
+		auto ValidateChunkId = [&](const FIoChunkId& ChunkId, const FCacheChunkInfo& ChunkInfo)
+			{
+				const uint8* ChunkdIdBytes = ChunkId.GetData();
+				const int32 IdByteCount = ChunkId.GetSize() - sizeof(decltype(ChunkId.GetChunkType()));
+
+				if (!FMemory::Memcmp(ChunkdIdBytes, FIoChunkId::InvalidChunkId.GetData(), IdByteCount))
+				{
+					InvalidChunkIds.Add(ChunkId);
+				}
 
+				if (ChunkIdSet.Contains(ChunkId))
+				{
+					DuplicateChunkIds.Add(ChunkId);
+				}
+				else
+				{
+					ChunkIdSet.Add(ChunkId);
+				}
+			};
+		
+		IterateChunkIds(ValidateChunkId);
 		IterateCacheEntries(ValidateEntryInfo);
 		for (const FIoChunkId& ChunkId : InvalidChunkIds)
 		{
 			Invalidate(ChunkId);
 		}
 
-		UE_LOG(LogCacheJournal, Warning, TEXT("There were %d invalid zen cache journal entries. All have been invalidated."), InvalidChunkIds.Num());
+		for (const FIoChunkId& ChunkId : DuplicateChunkIds)
+		{
+			Invalidate(ChunkId);
+		}
+
+		if (InvalidChunkIds.Num() > 0)
+		{
+			UE_LOG(LogCacheJournal, Warning, TEXT("There were %d invalid zen cache journal entries."), InvalidChunkIds.Num());
+		}
+
+		if (DuplicateChunkIds.Num() > 0)
+		{
+			UE_LOG(LogCacheJournal, Warning, TEXT("There were %d duplicate chunk id entries."), DuplicateChunkIds.Num());
+		}
 	}
 
 	FJournalChunkInfoPage* FCacheJournalSectioned::FindOrAddChunkInfoPage(const FIoChunkId& ChunkId)
