Load existing shader archives from the cooked output directory as well as the shader intermediate directory, to avoid issues with incremental cook which does not retain the shader intermediate directory in Zen.

[RN] none

diff --git a/Engine/Source/Developer/IoStoreUtilities/Private/ZenCookArtifactReader.cpp b/Engine/Source/Developer/IoStoreUtilities/Private/ZenCookArtifactReader.cpp
--- a/Engine/Source/Developer/IoStoreUtilities/Private/ZenCookArtifactReader.cpp
+++ b/Engine/Source/Developer/IoStoreUtilities/Private/ZenCookArtifactReader.cpp
@@ -89,7 +89,7 @@ bool FZenCookArtifactReader::IterateDirectory(const TCHAR* Directory, IPlatformF
 		FString StandardDirectory;
 		if (MakeStorageServerPath(Directory, StandardDirectory))
 		{
-			return StorageServerPlatformFile->IterateDirectory(Directory, Visitor);
+			return StorageServerPlatformFile->IterateDirectory(*StandardDirectory, Visitor);
 		}
 	}
 #endif // !UE_BUILD_SHIPPING
diff --git a/Engine/Source/Runtime/RenderCore/Private/ShaderCodeLibrary.cpp b/Engine/Source/Runtime/RenderCore/Private/ShaderCodeLibrary.cpp
--- a/Engine/Source/Runtime/RenderCore/Private/ShaderCodeLibrary.cpp
+++ b/Engine/Source/Runtime/RenderCore/Private/ShaderCodeLibrary.cpp
@@ -2165,12 +2165,14 @@ struct FEditorShaderCodeArchive
 	{
 		check(LibraryName.Len() > 0);
 		static const FRegexPattern ShaderCodePattern(ShaderCodePatternStr);
+		static const FRegexPattern ShaderArchivePattern(ShaderArchivePatternStr);
 
 		const FString FormatNameStr = FormatName.ToString();
 		const FString ShaderIntermediateLocation = FPaths::ProjectSavedDir() / TEXT("Shaders") / FormatNameStr;
 
 		TArray<FString> ShaderFiles;
 		ShaderFindFiles(ShaderFiles, *ShaderIntermediateLocation, *ShaderExtension);
+		ShaderFindFiles(ShaderFiles, *OutputDir, *ShaderExtension);
 
 		for (const FString& ShaderFileName : ShaderFiles)
 		{
@@ -2185,6 +2187,18 @@ struct FEditorShaderCodeArchive
 					AddShaderCodeLibraryByName(OutputDir, FoundLibraryName, CodeStats);
 				}
 			}
+
+			FRegexMatcher FindArchiveFormat(ShaderArchivePattern, ShaderFileName);
+			if (FindArchiveFormat.FindNext())
+			{
+				const FString& FoundArchiveName = FindArchiveFormat.GetCaptureGroup(1);
+				const FString& FoundFormatName = FindArchiveFormat.GetCaptureGroup(2);
+
+				if (FoundArchiveName.StartsWith(LibraryName) && FoundFormatName.Equals(FormatNameStr, ESearchCase::IgnoreCase))
+				{
+					AddShaderCodeLibraryByName(OutputDir, FoundArchiveName, CodeStats);
+				}
+			}
 		}
 	}
 
diff --git a/Engine/Source/Runtime/StorageServerClient/Private/StorageServerPlatformFile.cpp b/Engine/Source/Runtime/StorageServerClient/Private/StorageServerPlatformFile.cpp
--- a/Engine/Source/Runtime/StorageServerClient/Private/StorageServerPlatformFile.cpp
+++ b/Engine/Source/Runtime/StorageServerClient/Private/StorageServerPlatformFile.cpp
@@ -1615,7 +1615,7 @@ bool FStorageServerPlatformFile::IterateDirectory(const TCHAR* Directory, IPlatf
 	return MergedVisitor.PerformMergedIteration(Directory,
 		[this](const TCHAR* Directory, FDirectoryVisitor& Visitor)
 		{
-			return LowerLevel->IterateDirectory(Directory, Visitor);
+			return LowerLevel ? LowerLevel->IterateDirectory(Directory, Visitor) : true;
 		},
 		[this](const TCHAR* Directory, FDirectoryVisitor& Visitor)
 		{
@@ -1638,7 +1638,7 @@ bool FStorageServerPlatformFile::IterateDirectoryRecursively(const TCHAR* Direct
 	return MergedVisitor.PerformMergedIteration(Directory,
 		[this](const TCHAR* Directory, FDirectoryVisitor& Visitor)
 		{
-			return LowerLevel->IterateDirectoryRecursively(Directory, Visitor);
+			return LowerLevel ? LowerLevel->IterateDirectoryRecursively(Directory, Visitor) : true;
 		},
 		[this](const TCHAR* Directory, FDirectoryVisitor& Visitor)
 		{
@@ -1662,7 +1662,7 @@ bool FStorageServerPlatformFile::IterateDirectoryStat(const TCHAR* Directory, FD
 	return MergedVisitor.PerformMergedIteration(Directory,
 		[this](const TCHAR* Directory, FDirectoryStatVisitor& Visitor)
 		{
-			return LowerLevel->IterateDirectoryStat(Directory, Visitor);
+			return LowerLevel ? LowerLevel->IterateDirectoryStat(Directory, Visitor) : true;
 		},
 		[this](const TCHAR* Directory, FDirectoryStatVisitor& Visitor)
 		{
