When using StorageServerPlatformFile (eg: streaming data from zenserver) we should always allow reading data from the local filesystem if the data is under the saved directory regardless of extension.

diff --git a/Engine/Source/Runtime/StorageServerClient/Private/StorageServerPlatformFile.cpp b/Engine/Source/Runtime/StorageServerClient/Private/StorageServerPlatformFile.cpp
--- a/Engine/Source/Runtime/StorageServerClient/Private/StorageServerPlatformFile.cpp
+++ b/Engine/Source/Runtime/StorageServerClient/Private/StorageServerPlatformFile.cpp
@@ -74,6 +74,21 @@ static FAutoConsoleVariableRef CVarPreferLocalForNonAssets(
 	ECVF_Default
 );
 
+static bool IsFileInDynamicPath(const FString InFilename)
+{
+	return FPaths::IsUnderDirectory(InFilename, FPaths::ProjectSavedDir()) ||
+		FPaths::IsUnderDirectory(InFilename, FPaths::EngineSavedDir()) ||
+		FPaths::IsUnderDirectory(InFilename, FPaths::GeneratedConfigDir()) ||
+		FPaths::IsUnderDirectory(InFilename, FPaths::GameAgnosticSavedDir()) ||
+		FPaths::IsUnderDirectory(InFilename, FPaths::EngineVersionAgnosticUserDir()) ||
+		FPaths::IsUnderDirectory(InFilename, FPaths::EngineUserDir()) ||
+		FPaths::IsUnderDirectory(InFilename, FPaths::ProjectUserDir()) ||
+		FPaths::IsUnderDirectory(InFilename, FPaths::EngineIntermediateDir()) ||
+		FPaths::IsUnderDirectory(InFilename, FPaths::ProjectIntermediateDir()) ||
+		FPaths::IsUnderDirectory(InFilename, FPaths::CloudDir()) ||
+		FPaths::IsUnderDirectory(InFilename, FPaths::GameDevelopersDir()) ||
+		(FPaths::HasProjectPersistentDownloadDir() && FPaths::IsUnderDirectory(InFilename, FPaths::ProjectPersistentDownloadDir()));
+}
 
 
 FStorageServerFileSystemTOC::~FStorageServerFileSystemTOC()
@@ -1934,6 +1949,16 @@ bool FStorageServerPlatformFile::IsNonServerFilenameAllowed(FStringView InFilena
 		{
 			FName Ext = FName(FPathViews::GetExtension(InFilename));
 			bAllowed = !ExcludedNonServerExtensions.Contains(Ext);
+			if (!bAllowed)
+			{
+				// There are portions of the directory tree we expect to be writable, so we should
+				// allow reads to happen from them via the local filesystem, even for file extensions
+				// that we otherwise consider the storage server to be authoritative for.
+				if (IsFileInDynamicPath(FString(InFilename)))
+				{
+					bAllowed = true;
+				}
+			}
 
 			UE_CLOG(!bAllowed, LogStorageServerPlatformFile, VeryVerbose,
 				TEXT("Access to file '%.*s' is limited to server contents due to file extension being listed in ExcludedNonServerExtensions."),
