- Changed HybridCookedEditor to use cooked server data instead of cooked client data, when running the editor with -server. This is because there are issues with NotLoadForServer causing imports to fail. Another option would be to ignore NotLoadForServer when running editor with -server, but that is scary!

diff --git a/Engine/Config/BaseEngine.ini b/Engine/Config/BaseEngine.ini
--- a/Engine/Config/BaseEngine.ini
+++ b/Engine/Config/BaseEngine.ini
@@ -2628,6 +2628,7 @@ bIsSupported=True
 ; game uses Client builds, set this to {Platform}Client. Or if you use MakeCookedEditor to 
 ; make full cooked editor builds, you could use {Platform}CookedEditor
 RuntimeTargetPlatform="{Platform}"
+RuntimeServerPlatform="{Platform}Server"
 
 [Zen.Plugin.AndroidPortForwarder]
 Name=Engine/Source/Developer/Android/AndroidZenServerPlugin/Binaries/Win64/AndroidPortForwarder2.dll
diff --git a/Engine/Source/Runtime/StorageServerClient/Private/StorageServerPlatformFile.cpp b/Engine/Source/Runtime/StorageServerClient/Private/StorageServerPlatformFile.cpp
--- a/Engine/Source/Runtime/StorageServerClient/Private/StorageServerPlatformFile.cpp
+++ b/Engine/Source/Runtime/StorageServerClient/Private/StorageServerPlatformFile.cpp
@@ -900,15 +900,18 @@ static FString GetCookedPlatformName()
 	FString PlatformName;
 	if (IsRunningHybridCookedEditor())
 	{
+		bool bIsServer = UE_SERVER || FParse::Param(FCommandLine::Get(), TEXT("server"));
+		const TCHAR* Key = bIsServer ? TEXT("RuntimeServerPlatform") : TEXT("RuntimeTargetPlatform");
+
 		// manually look in DefaultEngine.ini, or BaseEngine.ini, for the name of the platform to use to find the ue.projectstore file
 		FConfigFile Default;
 		Default.Read(*(FPaths::ProjectConfigDir() / TEXT("DefaultEngine.ini")));
-		if (!Default.GetString(TEXT("HybridCookedEditor"), TEXT("RuntimeTargetPlatform"), PlatformName))
+		if (!Default.GetString(TEXT("HybridCookedEditor"), Key, PlatformName))
 		{
 			FConfigFile Base;
 			Base.Read(*(FPaths::EngineConfigDir() / TEXT("BaseEngine.ini")));
 			// we expect this to always be found
-			verify(Base.GetString(TEXT("HybridCookedEditor"), TEXT("RuntimeTargetPlatform"), PlatformName));
+			verify(Base.GetString(TEXT("HybridCookedEditor"), Key, PlatformName));
 		}
 		PlatformName.ReplaceInline(TEXT("{Platform}"), ANSI_TO_TCHAR(FPlatformProperties::IniPlatformName()));
 	}
