FCookOnTheFlyPackageStoreBackend::DoesPackageExist now more accurately implements the behaviour of FPackageName::DoesPackageExist
- A list of on-disk packages is enumerated when clients connect to COTFS
- If a package was loaded/cooked we use the existing behaviour to test for existence
- Any packages that were not loaded/cooked will be tested against the on-disk packages on the server

diff --git a/Engine/Source/Editor/UnrealEd/Private/Cooker/IoStoreCookOnTheFlyRequestManager.cpp b/Engine/Source/Editor/UnrealEd/Private/Cooker/IoStoreCookOnTheFlyRequestManager.cpp
--- a/Engine/Source/Editor/UnrealEd/Private/Cooker/IoStoreCookOnTheFlyRequestManager.cpp
+++ b/Engine/Source/Editor/UnrealEd/Private/Cooker/IoStoreCookOnTheFlyRequestManager.cpp
@@ -47,15 +47,33 @@ public:
 		PackageIdToName.Add(FPackageId::FromName(PackageName), PackageName);
 	}
 
-	FName Get(FPackageId PackageId)
+	FName Get(FPackageId PackageId) const
 	{
 		FScopeLock _(&CriticalSection);
 		return PackageIdToName.FindRef(PackageId);
 	}
 
+	uint64 Num() const
+	{
+		FScopeLock _(&CriticalSection);
+		return PackageIdToName.Num();
+	}
+
+	/**
+	 * Iterate over all entries in the registry
+	 */
+	void ForEachPackage(TFunction<void(const FPackageId&, const FName&)> Visitor) const
+	{
+		FScopeLock _(&CriticalSection);
+		for (const TPair<FPackageId, FName>& Entry : PackageIdToName)
+		{
+			Visitor(Entry.Key, Entry.Value);
+		}
+	}
+
 private:
 	TMap<FPackageId, FName> PackageIdToName;
-	FCriticalSection CriticalSection;
+	mutable FCriticalSection CriticalSection;
 };
 
 class FIoStoreCookOnTheFlyRequestManager final
@@ -74,6 +92,7 @@ public:
 		ConnectionServer->OnRequest(ECookOnTheFlyMessage::GetCookedPackages).BindRaw(this, &FIoStoreCookOnTheFlyRequestManager::HandleClientRequest);
 		ConnectionServer->OnRequest(ECookOnTheFlyMessage::CookPackage).BindRaw(this, &FIoStoreCookOnTheFlyRequestManager::HandleClientRequest);
 		ConnectionServer->OnRequest(ECookOnTheFlyMessage::RecookPackages).BindRaw(this, &FIoStoreCookOnTheFlyRequestManager::HandleClientRequest);
+		ConnectionServer->OnRequest(ECookOnTheFlyMessage::GetAllOnDiskPackages).BindRaw(this, &FIoStoreCookOnTheFlyRequestManager::HandleClientRequest);
 
 		MessageEndpoint = FMessageEndpoint::Builder("FCookOnTheFly");
 	}
@@ -574,6 +593,9 @@ private:
 			case UE::Cook::ECookOnTheFlyMessage::RecookPackages:
 				bRequestOk = HandleRecookPackagesRequest(Connection.GetPlatformName(), Request, Response);
 				break;
+			case UE::Cook::ECookOnTheFlyMessage::GetAllOnDiskPackages:
+				bRequestOk = HandleGetAllOnDiskPackagesRequest(Connection.GetPlatformName(), Request, Response);
+				break;
 			default:
 				return false;
 		}
@@ -710,6 +732,44 @@ private:
 		return true;
 	}
 
+	bool HandleGetAllOnDiskPackagesRequest(const FName& PlatformName, const UE::Cook::FCookOnTheFlyRequest& Request, UE::Cook::FCookOnTheFlyResponse& Response)
+	{
+		using namespace UE::ZenCookOnTheFly::Messaging;
+
+		TRACE_CPUPROFILER_EVENT_SCOPE(CookOnTheFly::HandleGetAllOnDiskPackagesRequest);
+
+		if (PlatformName.IsNone())		// only accept this request from non-editor clients
+		{
+			UE_LOG(LogCookOnTheFly, Warning, TEXT("GetAllOnDiskPackages from editor client"));
+			Response.SetStatus(UE::Cook::ECookOnTheFlyMessageStatus::Error);
+			return true;
+		}
+		else
+		{
+			FPlatformContext& Context = GetContext(PlatformName);
+			Context.RequestCookSessionAndWait();		
+
+			// After this point the packages have been enumerated into FPlatformContext::PackageRegistry
+			check(bCookerGlobalsInitialized);		
+
+			FGetAllOnDiskPackagesRequest AllOnDiskPackages;
+			AllOnDiskPackages.PackageIds.Reserve(PackageRegistry.Num());
+			{
+				FScopeLock __(&Context.GetLock());
+				PackageRegistry.ForEachPackage([&AllOnDiskPackages](const FPackageId& ID, const FName& Name)
+					{
+						AllOnDiskPackages.PackageIds.Emplace(ID);
+					});
+			}
+
+			Response.SetBodyTo(MoveTemp(AllOnDiskPackages));
+		}
+
+		Response.SetStatus(UE::Cook::ECookOnTheFlyMessageStatus::Ok);
+
+		return true;
+	}
+
 	bool BroadcastMessage(const UE::Cook::FCookOnTheFlyMessage& Message, const FName& PlatformName = NAME_None)
 	{
 		using namespace UE::Cook;
diff --git a/Engine/Source/Runtime/CookOnTheFly/Internal/CookOnTheFly.h b/Engine/Source/Runtime/CookOnTheFly/Internal/CookOnTheFly.h
--- a/Engine/Source/Runtime/CookOnTheFly/Internal/CookOnTheFly.h
+++ b/Engine/Source/Runtime/CookOnTheFly/Internal/CookOnTheFly.h
@@ -48,6 +48,8 @@ enum class ECookOnTheFlyMessage : uint32
 	RecookPackages		= 0x70,
 	/* Legacy message for NetworkPlatformFile */
 	NetworkPlatformFile	= 0x80,
+	/* Request for a list of on-disk packages, which may or may not be cooked or cookable for a given platform */
+	GetAllOnDiskPackages = 0x90
 };
 ENUM_CLASS_FLAGS(ECookOnTheFlyMessage);
 
diff --git a/Engine/Source/Runtime/CookOnTheFly/Internal/CookOnTheFlyMessages.h b/Engine/Source/Runtime/CookOnTheFly/Internal/CookOnTheFlyMessages.h
--- a/Engine/Source/Runtime/CookOnTheFly/Internal/CookOnTheFlyMessages.h
+++ b/Engine/Source/Runtime/CookOnTheFly/Internal/CookOnTheFlyMessages.h
@@ -37,4 +37,11 @@ struct FRecookPackagesResponse
 	COOKONTHEFLY_API friend FArchive& operator<<(FArchive& Ar, FRecookPackagesResponse& Response);
 };
 
+struct FGetAllOnDiskPackagesRequest
+{
+	TArray<FPackageId> PackageIds;
+
+	COOKONTHEFLY_API friend FArchive& operator<<(FArchive& Ar, FGetAllOnDiskPackagesRequest& Request);
+};
+
 }}} // namesapce UE::ZenCookOnTheFly::Messaging
diff --git a/Engine/Source/Runtime/CookOnTheFly/Private/CookOnTheFlyMessages.cpp b/Engine/Source/Runtime/CookOnTheFly/Private/CookOnTheFlyMessages.cpp
--- a/Engine/Source/Runtime/CookOnTheFly/Private/CookOnTheFlyMessages.cpp
+++ b/Engine/Source/Runtime/CookOnTheFly/Private/CookOnTheFlyMessages.cpp
@@ -32,4 +32,10 @@ FArchive& operator<<(FArchive& Ar, FRecookPackagesResponse& Response)
 	return Ar;
 }
 
+FArchive& operator<<(FArchive& Ar, FGetAllOnDiskPackagesRequest& Request)
+{ 
+	Ar << Request.PackageIds;
+	return Ar;
+}
+
 }}} // namesapce UE::ZenCookOnTheFly::Messaging
diff --git a/Engine/Source/Runtime/StorageServerClient/Private/CookOnTheFlyPackageStore.cpp b/Engine/Source/Runtime/StorageServerClient/Private/CookOnTheFlyPackageStore.cpp
--- a/Engine/Source/Runtime/StorageServerClient/Private/CookOnTheFlyPackageStore.cpp
+++ b/Engine/Source/Runtime/StorageServerClient/Private/CookOnTheFlyPackageStore.cpp
@@ -18,17 +18,26 @@ FCookOnTheFlyPackageStoreBackend::FCookOnTheFlyPackageStoreBackend(UE::Cook::ICo
 	PackageEntries.Add();
 	CookOnTheFlyServerConnection.OnMessage().AddRaw(this, &FCookOnTheFlyPackageStoreBackend::OnCookOnTheFlyMessage);
 
+	GetCookedPackagesFromServer();
+	GetAllAvailablePackagesFromServer();
+
+	FPackageName::DoesPackageExistOverride().BindLambda([this](FName PackageName)
+	{
+		return DoesPackageExist(FPackageId::FromName(PackageName));
+	});
+}
+
+void FCookOnTheFlyPackageStoreBackend::GetCookedPackagesFromServer()
+{
 	using namespace UE::Cook;
 	using namespace UE::ZenCookOnTheFly::Messaging;
 
 	FCookOnTheFlyRequest Request(ECookOnTheFlyMessage::GetCookedPackages);
-	
 	FCookOnTheFlyResponse Response = CookOnTheFlyServerConnection.SendRequest(Request).Get();
-
 	if (Response.IsOk())
 	{
 		FCompletedPackages GetCookedPackagesResponse = Response.GetBodyAs<FCompletedPackages>();
-		
+
 		UE_LOG(LogCookOnTheFly, Log, TEXT("Got '%d' cooked and '%d' failed packages from server"),
 			GetCookedPackagesResponse.CookedPackages.Num(), GetCookedPackagesResponse.FailedPackages.Num());
 
@@ -41,11 +50,26 @@ FCookOnTheFlyPackageStoreBackend::FCookOnTheFlyPackageStoreBackend(UE::Cook::ICo
 	{
 		UE_LOG(LogCookOnTheFly, Warning, TEXT("Failed to send 'GetCookedPackages' request"));
 	}
+}
 
-	FPackageName::DoesPackageExistOverride().BindLambda([this](FName PackageName)
+void FCookOnTheFlyPackageStoreBackend::GetAllAvailablePackagesFromServer()
+{
+	using namespace UE::Cook;
+	using namespace UE::ZenCookOnTheFly::Messaging;
+
+	FCookOnTheFlyRequest Request(ECookOnTheFlyMessage::GetAllOnDiskPackages);
+	FCookOnTheFlyResponse Response = CookOnTheFlyServerConnection.SendRequest(Request).Get();
+	if (Response.IsOk())
 	{
-		return DoesPackageExist(FPackageId::FromName(PackageName));
-	});
+		FGetAllOnDiskPackagesRequest GetAllOnDiskPackagesResponse = Response.GetBodyAs<FGetAllOnDiskPackagesRequest>();
+		UE_LOG(LogCookOnTheFly, Display, TEXT("Received list of %d total on-disk packages from server"), GetAllOnDiskPackagesResponse.PackageIds.Num());
+
+		AddOnDiskPackages(GetAllOnDiskPackagesResponse.PackageIds);
+	}
+	else
+	{
+		UE_LOG(LogCookOnTheFly, Warning, TEXT("Failed to send 'GetAllOnDiskPackages' request"));
+	}
 }
 
 void FCookOnTheFlyPackageStoreBackend::BeginRead()
@@ -67,8 +91,16 @@ void FCookOnTheFlyPackageStoreBackend::EndRead()
 bool FCookOnTheFlyPackageStoreBackend::DoesPackageExist(FPackageId PackageId)
 {
 	FReadScopeLock ReadLock(EntriesLock);
-	FEntryInfo EntryInfo = PackageIdToEntryInfo.FindRef(PackageId);
-	return EntryInfo.Status != EPackageStoreEntryStatus::Missing;
+	const FEntryInfo* EntryInfo = PackageIdToEntryInfo.Find(PackageId);
+	if (EntryInfo == nullptr)	
+	{
+		// We didn't encounter a cooked version of this package yet, but it may still exist on the server
+		return AllServerOnDiskPackages.Find(PackageId) != nullptr;
+	}
+	else
+	{
+		return EntryInfo->Status != EPackageStoreEntryStatus::Missing;
+	}
 }
 
 EPackageStoreEntryStatus FCookOnTheFlyPackageStoreBackend::GetPackageStoreEntry(FPackageId PackageId,
@@ -192,6 +224,16 @@ void FCookOnTheFlyPackageStoreBackend::AddPackages(TArray<FPackageStoreEntryReso
 	}
 }
 
+void FCookOnTheFlyPackageStoreBackend::AddOnDiskPackages(const TArray<FPackageId>& Packages)
+{
+	FWriteScopeLock WriteLock(EntriesLock);
+	AllServerOnDiskPackages.Reserve(Packages.Num());
+	for (const FPackageId& PackageID : Packages)
+	{
+		AllServerOnDiskPackages.Emplace(PackageID);
+	}
+}
+
 void FCookOnTheFlyPackageStoreBackend::OnCookOnTheFlyMessage(const UE::Cook::FCookOnTheFlyMessage& Message)
 {
 	using namespace UE::ZenCookOnTheFly::Messaging;
diff --git a/Engine/Source/Runtime/StorageServerClient/Private/CookOnTheFlyPackageStore.h b/Engine/Source/Runtime/StorageServerClient/Private/CookOnTheFlyPackageStore.h
--- a/Engine/Source/Runtime/StorageServerClient/Private/CookOnTheFlyPackageStore.h
+++ b/Engine/Source/Runtime/StorageServerClient/Private/CookOnTheFlyPackageStore.h
@@ -37,6 +37,10 @@ public:
 
 	virtual void EndRead() override;
 
+	/**
+	 *	Returns true if a cooked or uncooked package exists on-disk on the COTF server
+	 *	Note this may not be accurate for generated packages 
+	 */
 	bool DoesPackageExist(FPackageId PackageId);
 	virtual EPackageStoreEntryStatus GetPackageStoreEntry(FPackageId PackageId, FName PackageName,
 		FPackageStoreEntry& OutPackageStoreEntry) override;
@@ -47,10 +51,13 @@ public:
 	}
 	
 private:
+	void GetCookedPackagesFromServer();
+	void GetAllAvailablePackagesFromServer();
 	void SendCookRequest(TArray<TPair<FPackageId, FName>> PackageIds);
 	EPackageStoreEntryStatus CreatePackageStoreEntry(const FEntryInfo& EntryInfo, FPackageStoreEntry& OutPackageStoreEntry);
 	void AddPackages(TArray<FPackageStoreEntryResource> Entries, TArray<FPackageId> FailedPackageIds,
 		TArray<TPair<FPackageId, FName>> PackageIdsAndNames);
+	void AddOnDiskPackages(const TArray<FPackageId>& Packages);
 	void OnCookOnTheFlyMessage(const UE::Cook::FCookOnTheFlyMessage& Message);
 	void CheckActivity();
 
@@ -60,6 +67,11 @@ private:
 	TMap<FPackageId, FEntryInfo> PackageIdToEntryInfo;
 	TChunkedArray<FPackageStoreEntryResource> PackageEntries;
 	TArray<TPair<FPackageId, FName>> RequestedPackageIds;
+	/**
+	 * All packages available on-disk on the COTF server. 
+	 * Populated from server asset registry when connection is established. These may be cooked or uncooked
+	 */ 
+	TSet<FPackageId> AllServerOnDiskPackages;	
 	FPackageStats PackageStats;
 
 	const double MaxInactivityTime = 20;
