Fix Server build warning C6286: (<non-zero constant> || <expression>) is always a non-zero constant.  <expression> is never evaluated and might have side effects.
#rnx

diff --git a/Engine/Source/Runtime/StorageServerClient/Private/StorageServerPlatformFile.cpp b/Engine/Source/Runtime/StorageServerClient/Private/StorageServerPlatformFile.cpp
--- a/Engine/Source/Runtime/StorageServerClient/Private/StorageServerPlatformFile.cpp
+++ b/Engine/Source/Runtime/StorageServerClient/Private/StorageServerPlatformFile.cpp
@@ -900,8 +900,11 @@ static FString GetCookedPlatformName()
 	FString PlatformName;
 	if (IsRunningHybridCookedEditor())
 	{
-		bool bIsServer = UE_SERVER || FParse::Param(FCommandLine::Get(), TEXT("server"));
-		const TCHAR* Key = bIsServer ? TEXT("RuntimeServerPlatform") : TEXT("RuntimeTargetPlatform");
+#if UE_SERVER
+		const TCHAR* Key = TEXT("RuntimeServerPlatform");
+#else
+		const TCHAR* Key = FParse::Param(FCommandLine::Get(), TEXT("server")) ? TEXT("RuntimeServerPlatform") : TEXT("RuntimeTargetPlatform");
+#endif
 
 		// manually look in DefaultEngine.ini, or BaseEngine.ini, for the name of the platform to use to find the ue.projectstore file
 		FConfigFile Default;
