IncrementalCook: StorageServerPlatformFile: Set Connection to nullptr if Initialize fails; Initialize failure means the Connection will crash if used; StorageServerPlatformFile already handles graceful fallback noops anderror reporting if Connection is nuillptr.
Downgrade the log message for the failed to connect case from Warning to Display; it is expected to occur in the case of a new cook when we try to connect the ZenCookArtifactReader before the Writer has created the oplog. A non-existent oplog should be treated the same as no files available for download.
#rnx

diff --git a/Engine/Source/Runtime/StorageServerClient/Private/StorageServerPlatformFile.cpp b/Engine/Source/Runtime/StorageServerClient/Private/StorageServerPlatformFile.cpp
--- a/Engine/Source/Runtime/StorageServerClient/Private/StorageServerPlatformFile.cpp
+++ b/Engine/Source/Runtime/StorageServerClient/Private/StorageServerPlatformFile.cpp
@@ -1417,12 +1417,18 @@ void FStorageServerPlatformFile::InitializeConnection()
 			FPlatformMisc::MessageBoxExt(EAppMsgType::Ok, *FailedConnectionText, *FailedConnectionTitle);
 		}
 
-		UE_LOG(LogStorageServerPlatformFile, Error, TEXT("Failed to initialize connection to %s"), *FString::Join(HostAddrs, TEXT("\n")));
+		UE_LOG(LogStorageServerPlatformFile, Error, TEXT("Failed to initialize connection to any of [ %s ]"),
+			*FString::Join(HostAddrs, TEXT(", ")));
+		Connection.Reset();
 		FPlatformMisc::RequestExit(true);
 	}
 	else
 	{
-		UE_LOG(LogStorageServerPlatformFile, Warning, TEXT("Failed to initialize connection to %s"), *FString::Join(HostAddrs, TEXT("\n")));
+		// One reason the connection can fail is if the oplog has not yet been created during a new cook:
+		// we will get the message "oplog <PlatformName> not found in project <ProjectName>".
+		UE_LOG(LogStorageServerPlatformFile, Display, TEXT("Failed to initialize connection to any of [ %s ]"),
+			*FString::Join(HostAddrs, TEXT(", ")));
+		Connection.Reset();
 	}
 }
 
