Remove Zen host addresses that refer to a Unix Domain Socket. FStorageServerPlatformFile does not support it currently.

#Synced-CL 48177787

diff --git a/Engine/Source/Runtime/StorageServerClient/Private/StorageServerPlatformFile.cpp b/Engine/Source/Runtime/StorageServerClient/Private/StorageServerPlatformFile.cpp
--- a/Engine/Source/Runtime/StorageServerClient/Private/StorageServerPlatformFile.cpp
+++ b/Engine/Source/Runtime/StorageServerClient/Private/StorageServerPlatformFile.cpp
@@ -1225,6 +1225,11 @@ bool FStorageServerPlatformFile::ShouldBeUsed(IPlatformFile* Inner, const TCHAR*
 	// add hostnames as last resort
 	HostAddrs.Append(HostNames);
 
+	if (HostAddrs.RemoveAll([](const FString& Val) {return Val.EndsWith(".sock"); }))
+	{
+		UE_LOG(LogStorageServerPlatformFile, Display, TEXT("Removed Zen host addresses that refer to a Unix Domain Socket. FStorageServerPlatformFile does not support UDS currently."));
+	}
+
 	if (!bPreferFileSystem || UE::IsUsingZenPakFileStreaming())
 	{
 		return HostAddrs.Num() > 0;
