[HttpClient]
- Fix buffer overflow by providing length for the string

#Synced-CL 49229088

diff --git a/Engine/Source/Runtime/StorageServerClient/Private/BuiltInHttpClient/BuiltInHttpClient.cpp b/Engine/Source/Runtime/StorageServerClient/Private/BuiltInHttpClient/BuiltInHttpClient.cpp
--- a/Engine/Source/Runtime/StorageServerClient/Private/BuiltInHttpClient/BuiltInHttpClient.cpp
+++ b/Engine/Source/Runtime/StorageServerClient/Private/BuiltInHttpClient/BuiltInHttpClient.cpp
@@ -199,7 +199,7 @@ IStorageServerHttpClient::FResult FBuiltInHttpClient::RequestSync(
 	}
 	else
 	{
-		FString ErrorMessage = bHasResponsePayload ? FString::ConstructFromPtrSize(ANSI_TO_TCHAR(reinterpret_cast<ANSICHAR*>(ResponsePayload->GetData())), ResponsePayload->GetSize()) : TEXT("Unknown error");
+		FString ErrorMessage = bHasResponsePayload ? FString::ConstructFromPtrSize(reinterpret_cast<const ANSICHAR*>(ResponsePayload->GetData()), ResponsePayload->GetSize()) : TEXT("Unknown error");
 
 		return {TIoStatusOr<FIoBuffer>(FIoStatus(ErrorCode, ErrorMessage)), ContentType};
 	}
