ZenArtifactReader unix domain sockets support

#Synced-CL 49721519

diff --git a/Engine/Source/Runtime/StorageServerClient/Private/BuiltInHttpClient/BuiltInHttpClient.cpp b/Engine/Source/Runtime/StorageServerClient/Private/BuiltInHttpClient/BuiltInHttpClient.cpp
--- a/Engine/Source/Runtime/StorageServerClient/Private/BuiltInHttpClient/BuiltInHttpClient.cpp
+++ b/Engine/Source/Runtime/StorageServerClient/Private/BuiltInHttpClient/BuiltInHttpClient.cpp
@@ -121,10 +121,10 @@ IStorageServerHttpClient::FResult FBuiltInHttpClient::RequestSync(
 	{
 		for (;;)
 		{
-			uint64 BytesRead;
+			uint64 BytesRead = 0;
 			Socket->Recv(ReadResponseLineBuffer, 1024, BytesRead, ESocketReceiveFlags::Peek);
 			FAnsiStringView ResponseView(reinterpret_cast<const ANSICHAR*>(ReadResponseLineBuffer), BytesRead);
-			int32 LineEndIndex;
+			int32 LineEndIndex = 0;
 			if (ResponseView.FindChar('\r', LineEndIndex) && BytesRead >= LineEndIndex + 2)
 			{
 				check(ResponseView[LineEndIndex + 1] == '\n');
diff --git a/Engine/Source/Runtime/StorageServerClient/Private/BuiltInHttpClient/BuiltInHttpClientUDSSocket.cpp b/Engine/Source/Runtime/StorageServerClient/Private/BuiltInHttpClient/BuiltInHttpClientUDSSocket.cpp
new file mode 100644
--- /dev/null
+++ b/Engine/Source/Runtime/StorageServerClient/Private/BuiltInHttpClient/BuiltInHttpClientUDSSocket.cpp
@@ -0,0 +1,190 @@
+// Copyright Epic Games, Inc. All Rights Reserved.
+
+#include "BuiltInHttpClientUDSSocket.h"
+#if !UE_BUILD_SHIPPING && PLATFORM_HAS_BSD_SOCKET_FEATURE_WINSOCKETS
+#include "StorageServerConnection.h"
+#include "Windows/AllowWindowsPlatformTypes.h"
+
+#include <winsock2.h>
+#include <ws2tcpip.h>
+#include <afunix.h>
+
+
+FBuiltInHttpClientUDSSocket::FBuiltInHttpClientUDSSocket(uint64 InSocket)
+	: Socket(InSocket)
+{
+}
+
+FBuiltInHttpClientUDSSocket::~FBuiltInHttpClientUDSSocket()
+{
+	if (Socket != INVALID_SOCKET)
+	{
+		closesocket(Socket);
+	}
+}
+bool FBuiltInHttpClientUDSSocket::Send(
+	const uint8* Data,
+	const uint64 DataSize)
+{
+	if (Socket == INVALID_SOCKET)
+	{
+		return false;
+	}
+
+	if (DataSize > INT_MAX)
+	{
+		return false;
+	}
+
+	int32 TotalBytesSent = 0;
+	while (TotalBytesSent < (int64)DataSize)
+	{
+		int32 BytesSent = send(Socket, (const char*)Data, (int)DataSize, 0);
+		if (BytesSent < 0)
+		{
+			return false;
+		}
+		TotalBytesSent += BytesSent;
+	}
+
+	return TotalBytesSent >= 0;
+}
+
+bool FBuiltInHttpClientUDSSocket::Recv(
+	uint8* Data,
+	const uint64 DataSize,
+	uint64& BytesRead,
+	ESocketReceiveFlags::Type ReceiveFlags)
+{
+	BytesRead = 0;
+
+	if (Socket == INVALID_SOCKET)
+	{
+		return false;
+	}
+	
+	if (DataSize > INT_MAX)
+	{
+		return false;
+	}
+
+	int ReadBytes = recv(Socket, (char*)Data, (int)DataSize, ReceiveFlags);
+
+	// 0 means the socket is closed gracefully no more data incoming. 
+	if (ReadBytes >= 0)
+	{
+		BytesRead = (uint64)ReadBytes;
+		return true;
+	}
+	
+	ReadBytes = 0;
+
+	int Code = WSAGetLastError();
+	if (Code == WSAEWOULDBLOCK)
+		return true;
+
+	return false;
+}
+
+bool FBuiltInHttpClientUDSSocket::HasPendingData(
+	uint64& PendingDataSize) const
+{
+	PendingDataSize = 0;
+	if (Socket == INVALID_SOCKET)
+	{
+		return false;
+	}
+
+#if PLATFORM_HAS_BSD_SOCKET_FEATURE_IOCTL
+	// See if there is any pending data on the read socket
+	uint32 DataAvailable = 0;
+	if (ioctlsocket(Socket, FIONREAD, (u_long*)(&DataAvailable)) == 0)
+	{
+		PendingDataSize = DataAvailable;
+		return (PendingDataSize > 0);
+	}
+	else
+	{
+		UE_LOG(LogStorageServerConnection, Fatal, TEXT("Call to obtain the pending data on socket failed"));
+	}
+#endif
+
+	return false;
+}
+
+void FBuiltInHttpClientUDSSocket::Close()
+{
+	if (Socket != INVALID_SOCKET)
+	{
+		closesocket(Socket);
+		Socket = INVALID_SOCKET;
+	}
+}
+
+FBuiltInHttpClientUDSSocketPool::FBuiltInHttpClientUDSSocketPool(
+	FString InServerAddr)
+	: Address(InServerAddr)
+{
+}
+
+FBuiltInHttpClientUDSSocketPool::~FBuiltInHttpClientUDSSocketPool()
+{
+	IBuiltInHttpClientSocket* Socket = nullptr;
+	while ((Socket = SocketPool.Pop()) != nullptr)
+	{
+		delete Socket;
+	}
+}
+
+IBuiltInHttpClientSocket* FBuiltInHttpClientUDSSocketPool::AcquireSocket(
+	float /*TimeoutSeconds*/)
+{
+	IBuiltInHttpClientSocket* SocketWrapper = SocketPool.Pop();
+	if (SocketWrapper)
+	{
+		return SocketWrapper;
+	}
+
+	SOCKET ListenSocket = socket(AF_UNIX, SOCK_STREAM, 0);
+	if (ListenSocket == INVALID_SOCKET)
+	{
+		UE_LOG(LogStorageServerConnection, Error, TEXT("failed to create AF_UNIX socket"));
+		return nullptr;
+	}
+
+	SOCKADDR_UN SocketAddress = { 0 };
+	SocketAddress.sun_family = AF_UNIX;
+	TExternalStringBuilder<char> Builder(SocketAddress.sun_path, UNIX_PATH_MAX);
+	Builder.Append(*Address, Address.Len());
+
+	if (connect(ListenSocket, (struct sockaddr*)&SocketAddress, sizeof(SOCKADDR_UN)) == -1)
+	{
+		UE_LOG(LogStorageServerConnection, Error, TEXT("connect failed: %s"), *Address);
+		closesocket(ListenSocket);
+		return nullptr;
+	}
+
+	return new FBuiltInHttpClientUDSSocket(ListenSocket);
+}
+
+void FBuiltInHttpClientUDSSocketPool::ReleaseSocket(
+	IBuiltInHttpClientSocket* Socket,
+	bool bKeepAlive)
+{
+	uint64 PendingDataSize = 0;
+	if (bKeepAlive && !Socket->HasPendingData(PendingDataSize))
+	{
+		SocketPool.Push(Socket);
+	}
+	else
+	{
+		if (PendingDataSize > 0)
+		{
+			UE_LOG(LogStorageServerConnection, Fatal, TEXT("Socket was not fully drained"));
+		}
+
+		delete Socket;
+	}
+}
+
+#endif
\ No newline at end of file
diff --git a/Engine/Source/Runtime/StorageServerClient/Private/BuiltInHttpClient/BuiltInHttpClientUDSSocket.h b/Engine/Source/Runtime/StorageServerClient/Private/BuiltInHttpClient/BuiltInHttpClientUDSSocket.h
new file mode 100644
--- /dev/null
+++ b/Engine/Source/Runtime/StorageServerClient/Private/BuiltInHttpClient/BuiltInHttpClientUDSSocket.h
@@ -0,0 +1,53 @@
+// Copyright Epic Games, Inc. All Rights Reserved.
+
+#pragma once
+
+#include "BuiltInHttpClient.h"
+#include "Containers/LockFreeList.h"
+
+#if !UE_BUILD_SHIPPING && PLATFORM_HAS_BSD_SOCKET_FEATURE_WINSOCKETS
+
+class FBuiltInHttpClientUDSSocket : public IBuiltInHttpClientSocket
+{
+public:
+	FBuiltInHttpClientUDSSocket(uint64 InSocket);
+	virtual ~FBuiltInHttpClientUDSSocket() override;
+
+	virtual bool Send(
+		const uint8* Data, 
+		const uint64 DataSize) override;
+
+	virtual bool Recv(
+		uint8* Data, 
+		const uint64 DataSize, 
+		uint64& BytesRead, 
+		ESocketReceiveFlags::Type ReceiveFlags) override;
+
+	virtual bool HasPendingData(
+		uint64& PendingDataSize) const override;
+	
+	virtual void Close() override;
+
+private:
+	uint64 Socket;
+};
+
+class FBuiltInHttpClientUDSSocketPool : public IBuiltInHttpClientSocketPool
+{
+public:
+	FBuiltInHttpClientUDSSocketPool(FString InServerAddr);
+	virtual ~FBuiltInHttpClientUDSSocketPool() override;
+
+	virtual IBuiltInHttpClientSocket* AcquireSocket(
+		float TimeoutSeconds = -1.f) override;
+	
+	virtual void ReleaseSocket(
+		IBuiltInHttpClientSocket* Socket, 
+		bool bKeepAlive) override;
+
+private:
+	const FString Address;
+	TLockFreePointerListUnordered<IBuiltInHttpClientSocket, PLATFORM_CACHE_LINE_SIZE> SocketPool;
+};
+
+#endif
\ No newline at end of file
diff --git a/Engine/Source/Runtime/StorageServerClient/Private/StorageServerConnection.cpp b/Engine/Source/Runtime/StorageServerClient/Private/StorageServerConnection.cpp
--- a/Engine/Source/Runtime/StorageServerClient/Private/StorageServerConnection.cpp
+++ b/Engine/Source/Runtime/StorageServerClient/Private/StorageServerConnection.cpp
@@ -21,6 +21,9 @@
 #include "BuiltInHttpClient/BuiltInHttpClient.h"
 #include "BuiltInHttpClient/BuiltInHttpClientFSocket.h"
 #include "BuiltInHttpClient/BuiltInHttpClientPlatformSocket.h"
+#if PLATFORM_HAS_BSD_SOCKET_FEATURE_WINSOCKETS
+#include "BuiltInHttpClient/BuiltInHttpClientUDSSocket.h"
+#endif
 #include "HAL/PlatformMath.h"
 #include "HAL/RunnableThread.h"
 #include "HAL/Event.h"
@@ -834,6 +837,11 @@ bool FStorageServerConnection::IsHostnameAddress(const FString Address)
 	return Address.StartsWith(TEXT("hostname://"));
 }
 
+bool FStorageServerConnection::IsUDSSocketAddress(const FString Address)
+{
+	return Address.EndsWith(TEXT(".sock"));
+}
+
 TUniquePtr<IStorageServerHttpClient> FStorageServerConnection::CreateHttpClient(const FString Address, const int32 Port)
 {
 	TSharedPtr<FInternetAddr> Addr = StringToInternetAddr(Address, Port);
@@ -848,6 +856,16 @@ TUniquePtr<IStorageServerHttpClient> FStorageServerConnection::CreateHttpClient(
 	{
 		SocketPool = MakeUnique<FBuiltInHttpClientPlatformSocketPool>(Address); 
 	}
+#if PLATFORM_HAS_BSD_SOCKET_FEATURE_WINSOCKETS
+	else if (IsUDSSocketAddress(Address))
+	{
+		SocketPool = MakeUnique<FBuiltInHttpClientUDSSocketPool>(Address);
+		// The request will be forwarded by the system that reads from the socket to the right target server
+		// Hence the host name in this case doesn't matter. We overwrite it to something basic such as localhost
+		// so that the socket path such as 'I:/sockets/zenstore.sock' doesn't mess up the URI generation. 
+		HostName = "localhost"; 
+	}
+#endif
 	else
 	{
 		SocketPool = MakeUnique<FBuiltInHttpClientFSocketPool>(Addr, *ISocketSubsystem::Get());
@@ -866,6 +884,10 @@ TSharedPtr<FInternetAddr> FStorageServerConnection::StringToInternetAddr(const F
 	{
 		return Result;
 	}
+	else if (IsUDSSocketAddress(HostAddr))
+	{
+		return Result;
+	}
 
 	// Numeric IPV6 addresses can be enclosed in brackets, and must have the brackets stripped before calling GetAddressFromString
 	FString ModifiedHostAddr;
diff --git a/Engine/Source/Runtime/StorageServerClient/Private/StorageServerConnection.h b/Engine/Source/Runtime/StorageServerClient/Private/StorageServerConnection.h
--- a/Engine/Source/Runtime/StorageServerClient/Private/StorageServerConnection.h
+++ b/Engine/Source/Runtime/StorageServerClient/Private/StorageServerConnection.h
@@ -147,7 +147,8 @@ private:
 	std::atomic<double> MaxRequestThroughput = -DBL_MAX;
 
 	TArray<FString> SortHostAddressesByLocalSubnet(TArrayView<const FString> HostAddresses, const int32 Port);
-	static bool IsPlatformSocketAddress(const FString Address); 
+	static bool IsPlatformSocketAddress(const FString Address);
+	static bool IsUDSSocketAddress(const FString Address);
 	static bool IsHostnameAddress(const FString Address); 
 	TUniquePtr<IStorageServerHttpClient> CreateHttpClient(const FString Address, const int32 Port);
 	TSharedPtr<FInternetAddr> StringToInternetAddr(const FString Address, const int32 Port);
diff --git a/Engine/Source/Runtime/StorageServerClient/Private/StorageServerPlatformFile.cpp b/Engine/Source/Runtime/StorageServerClient/Private/StorageServerPlatformFile.cpp
--- a/Engine/Source/Runtime/StorageServerClient/Private/StorageServerPlatformFile.cpp
+++ b/Engine/Source/Runtime/StorageServerClient/Private/StorageServerPlatformFile.cpp
@@ -1225,11 +1225,6 @@ bool FStorageServerPlatformFile::ShouldBeUsed(IPlatformFile* Inner, const TCHAR*
 	// add hostnames as last resort
 	HostAddrs.Append(HostNames);
 
-	if (HostAddrs.RemoveAll([](const FString& Val) {return Val.EndsWith(".sock"); }))
-	{
-		UE_LOG(LogStorageServerPlatformFile, Display, TEXT("Removed Zen host addresses that refer to a Unix Domain Socket. FStorageServerPlatformFile does not support UDS currently."));
-	}
-
 	if (!bPreferFileSystem || UE::IsUsingZenPakFileStreaming())
 	{
 		return HostAddrs.Num() > 0;
