domain socket flag translation fix.

MSG_WAITALL is 0x8 on Windows, but 0x100 on other platforms.

#Synced-CL 50193686

#Synced-CL 50193686

#Synced-CL 50193686

diff --git a/Engine/Source/Runtime/StorageServerClient/Private/BuiltInHttpClient/BuiltInHttpClientUDSSocket.cpp b/Engine/Source/Runtime/StorageServerClient/Private/BuiltInHttpClient/BuiltInHttpClientUDSSocket.cpp
--- a/Engine/Source/Runtime/StorageServerClient/Private/BuiltInHttpClient/BuiltInHttpClientUDSSocket.cpp
+++ b/Engine/Source/Runtime/StorageServerClient/Private/BuiltInHttpClient/BuiltInHttpClientUDSSocket.cpp
@@ -9,6 +9,26 @@
 #include <ws2tcpip.h>
 #include <afunix.h>
 
+// MSG_WAITALL is 0x8 on Windows, but 0x100 on other platforms.
+inline int TranslateFlags(ESocketReceiveFlags::Type Flags)
+{
+	int TranslatedFlags = 0;
+
+	if (Flags & ESocketReceiveFlags::Peek)
+	{
+		TranslatedFlags |= MSG_PEEK;
+#if PLATFORM_HAS_BSD_SOCKET_FEATURE_MSG_DONTWAIT
+		TranslatedFlags |= MSG_DONTWAIT;
+#endif // PLATFORM_HAS_BSD_SOCKET_FEATURE_MSG_DONTWAIT
+	}
+
+	if (Flags & ESocketReceiveFlags::WaitAll)
+	{
+		TranslatedFlags |= MSG_WAITALL;
+	}
+
+	return TranslatedFlags;
+}
 
 FBuiltInHttpClientUDSSocket::FBuiltInHttpClientUDSSocket(uint64 InSocket)
 	: Socket(InSocket)
@@ -68,7 +88,7 @@ bool FBuiltInHttpClientUDSSocket::Recv(
 		return false;
 	}
 
-	int ReadBytes = recv(Socket, (char*)Data, (int)DataSize, ReceiveFlags);
+	int ReadBytes = recv(Socket, (char*)Data, (int)DataSize, TranslateFlags(ReceiveFlags));
 
 	// 0 means the socket is closed gracefully no more data incoming. 
 	if (ReadBytes >= 0)
