Add missing HideWindowsPlatformTypes

#Synced-CL 51088011

diff --git a/Engine/Source/Runtime/StorageServerClient/Private/BuiltInHttpClient/BuiltInHttpClientUDSSocket.cpp b/Engine/Source/Runtime/StorageServerClient/Private/BuiltInHttpClient/BuiltInHttpClientUDSSocket.cpp
--- a/Engine/Source/Runtime/StorageServerClient/Private/BuiltInHttpClient/BuiltInHttpClientUDSSocket.cpp
+++ b/Engine/Source/Runtime/StorageServerClient/Private/BuiltInHttpClient/BuiltInHttpClientUDSSocket.cpp
@@ -208,4 +208,5 @@ void FBuiltInHttpClientUDSSocketPool::ReleaseSocket(
 	}
 }
 
-#endif
\ No newline at end of file
+#include "Windows/HideWindowsPlatformTypes.h"
+#endif // !UE_BUILD_SHIPPING && PLATFORM_HAS_BSD_SOCKET_FEATURE_WINSOCKETS
