Speculative fix for storage sever precache
- restore thread sync on destruction
- break out of the precache loop if we should exit immediately
- only a partial fix... if the ReadChunkRequest / ReadChunkBatchRequest stalls then it will be stuck until it completes
- avoid static deinit race with event pool cleanup if OnFEngineLoopInitComplete doesn't get called

#Synced-CL 51529686

diff --git a/Engine/Source/Runtime/StorageServerClient/Private/StorageServerPlatformFile.cpp b/Engine/Source/Runtime/StorageServerClient/Private/StorageServerPlatformFile.cpp
--- a/Engine/Source/Runtime/StorageServerClient/Private/StorageServerPlatformFile.cpp
+++ b/Engine/Source/Runtime/StorageServerClient/Private/StorageServerPlatformFile.cpp
@@ -2,6 +2,7 @@
 
 #include "StorageServerPlatformFile.h"
 #include "Algo/Replace.h"
+#include "Async/ManualResetEvent.h"
 #include "CookOnTheFly.h"
 #include "CookOnTheFlyPackageStore.h"
 #include "HAL/FileManagerGeneric.h"
@@ -282,15 +283,13 @@ public:
 	FStorageServerEngineStartupPrecache(FStorageServerConnection& InConnection)
 		: Connection(InConnection)
 		, CacheSize(0)
-		, HasWork(FPlatformProcess::GetSynchEventFromPool())
-		, IsCompleted(FPlatformProcess::GetSynchEventFromPool())
 		, bExitWhenComplete(false)
 		, bExitImmediately(false)
 		, bUseBatchedRequests(true)
 	{
 		if (FRunnableThread::Create( this, TEXT("StorageServerPrecache"), 0, EThreadPriority::TPri_Normal) == nullptr)
 		{
-			IsCompleted->Trigger();
+			IsCompleted.Notify();
 		}
 
 		if (Connection.IsConnectedToWorkspace())
@@ -301,6 +300,12 @@ public:
 
 	virtual ~FStorageServerEngineStartupPrecache()
 	{
+		// the worker thread has most likely completed already but no harm in shutting down cleanly
+		bExitImmediately = true;
+		HasWork.Notify();
+		IsCompleted.Wait();
+		IsCompleted.Reset();
+
 #if STORAGE_SERVER_START_CACHE_REPORT
 		UE_CLOG(Cache.Num() > 0, LogStorageServerPlatformFile, Log, TEXT("Engine startup precache size was %llu bytes, %d items. %d items actually retrieved (%d%%). %d cache misses (%d%%)"), CacheSize.load(), Cache.Num(), AccessedChunks.Num(), (AccessedChunks.Num()*100) / Cache.Num(), UncachedChunks.Num(), (UncachedChunks.Num()*100) / (AccessedChunks.Num()+UncachedChunks.Num()) );
 #endif
@@ -316,14 +321,14 @@ public:
 		{
 			PrecacheFileRequests.Enqueue( TTuple<FIoChunkId,int32>(Id, Size) );
 		}
-		HasWork->Trigger();
+		HasWork.Notify();
 	}
 
 	void Finalize()
 	{
 		// no more precache requests after this - thread will shut down when precaching has finished
 		bExitWhenComplete = true;
-		HasWork->Trigger();
+		HasWork.Notify();
 	}
 
 	bool GetPrecachedFile( const FIoChunkId& Id, int64 Offset, int64 Size, uint8* Destination, int64& BytesRead)
@@ -365,7 +370,8 @@ private:
 
 		while (!bExitImmediately)
 		{
-			HasWork->Wait();
+			HasWork.Wait();
+			HasWork.Reset();
 			if (bExitImmediately)
 			{
 				break;
@@ -373,7 +379,7 @@ private:
 
 #if HAS_STORAGE_SERVER_RPC_GETCHUNKS_API
 			PrecacheItems();
-			if (bExitWhenComplete)
+			if (bExitWhenComplete && !bExitImmediately)
 			{
 				// Clear out the final batch of requests which may have been added after the previous
 				// PrecacheItems() call had processed the queue, but before the HTTP response had
@@ -384,7 +390,7 @@ private:
 #endif
 		}
 
-		IsCompleted->Trigger();
+		IsCompleted.Notify();
 		return 0;
 	}
 
@@ -433,6 +439,11 @@ private:
 			const FIoChunkId& Id = Item.Key;
 			const int32& Size = Item.Value;
 
+			if (bExitImmediately)
+			{
+				return;
+			}
+
 			// make sure the total cache size doesn't grow out of control
 			if (CacheSize + Size > MaxCacheSize)
 			{
@@ -485,8 +496,8 @@ private:
 
 	TSpscQueue<TTuple<FIoChunkId,int32>> PrecacheFileRequests;
 	TSpscQueue<TTuple<FIoChunkId,int32>> HighPriorityPrecacheFileRequests;
-	FEvent* HasWork;
-	FEvent* IsCompleted;
+	UE::FManualResetEvent HasWork;
+	UE::FManualResetEvent IsCompleted;
 
 	std::atomic<bool> bExitWhenComplete;
 	std::atomic<bool> bExitImmediately;
