Turning off batch prefetching (but still allowing non-batch prefetching) of early startup data used during zen streaming.  The batch prefetch was showing up as a large stall for some users.

#Synced-CL 51529686

diff --git a/Engine/Source/Runtime/StorageServerClient/Private/StorageServerPlatformFile.cpp b/Engine/Source/Runtime/StorageServerClient/Private/StorageServerPlatformFile.cpp
--- a/Engine/Source/Runtime/StorageServerClient/Private/StorageServerPlatformFile.cpp
+++ b/Engine/Source/Runtime/StorageServerClient/Private/StorageServerPlatformFile.cpp
@@ -285,7 +285,7 @@ public:
 		, CacheSize(0)
 		, bExitWhenComplete(false)
 		, bExitImmediately(false)
-		, bUseBatchedRequests(true)
+		, bUseBatchedRequests(false) // Turned off batch requests until a stall can be diagnosed and fixed
 	{
 		if (FRunnableThread::Create( this, TEXT("StorageServerPrecache"), 0, EThreadPriority::TPri_Normal) == nullptr)
 		{
