[StorageServerClient]
* Fixed non unity compile error

#Synced-CL 52124500

diff --git a/Engine/Source/Runtime/StorageServerClient/Private/CookOnTheFlyPackageStore.h b/Engine/Source/Runtime/StorageServerClient/Private/CookOnTheFlyPackageStore.h
--- a/Engine/Source/Runtime/StorageServerClient/Private/CookOnTheFlyPackageStore.h
+++ b/Engine/Source/Runtime/StorageServerClient/Private/CookOnTheFlyPackageStore.h
@@ -6,6 +6,7 @@
 
 #if WITH_COTF
 
+#include "Containers/ChunkedArray.h"
 #include "Serialization/PackageStore.h"
 #include "CookOnTheFly.h"
 #include "Misc/ScopeRWLock.h"
